fixed up incorrect/obsolete module dependencies in various UFE related modules

diff --git a/Engine/Source/Developer/DeviceManager/DeviceManager.Build.cs b/Engine/Source/Developer/DeviceManager/DeviceManager.Build.cs
--- a/Engine/Source/Developer/DeviceManager/DeviceManager.Build.cs
+++ b/Engine/Source/Developer/DeviceManager/DeviceManager.Build.cs
@@ -6,35 +6,33 @@ public class DeviceManager : ModuleRules
 {
 	public DeviceManager(TargetInfo Target)
 	{
-		PrivateIncludePaths.AddRange(
-			new string[]
-			{
-				"Developer/DeviceManager/Private",
-				"Developer/DeviceManager/Private/Models",
-				"Developer/DeviceManager/Private/Widgets",
-				"Developer/DeviceManager/Private/Widgets/Apps",
-				"Developer/DeviceManager/Private/Widgets/Browser",
-				"Developer/DeviceManager/Private/Widgets/Details",
-				"Developer/DeviceManager/Private/Widgets/Processes",
-				"Developer/DeviceManager/Private/Widgets/Shared",
-				"Developer/DeviceManager/Private/Widgets/Toolbar",
+		PublicIncludePathModuleNames.AddRange(
+			new string[] {
+				"TargetDeviceServices",
 			}
 		);
 
 		PrivateDependencyModuleNames.AddRange(
-			new string[]
-			{
+			new string[] {
 				"Core",
 				"EditorStyle",
                 "InputCore",
 				"Slate",
+				"TargetPlatform",
 			}
 		);
 
-		PublicDependencyModuleNames.AddRange(
-			new string[]
-			{
-				"TargetDeviceServices",
+		PrivateIncludePaths.AddRange(
+			new string[] {
+				"Developer/DeviceManager/Private",
+				"Developer/DeviceManager/Private/Models",
+				"Developer/DeviceManager/Private/Widgets",
+				"Developer/DeviceManager/Private/Widgets/Apps",
+				"Developer/DeviceManager/Private/Widgets/Browser",
+				"Developer/DeviceManager/Private/Widgets/Details",
+				"Developer/DeviceManager/Private/Widgets/Processes",
+				"Developer/DeviceManager/Private/Widgets/Shared",
+				"Developer/DeviceManager/Private/Widgets/Toolbar",
 			}
 		);
 	}
diff --git a/Engine/Source/Developer/DeviceManager/Private/DeviceManagerPrivatePCH.h b/Engine/Source/Developer/DeviceManager/Private/DeviceManagerPrivatePCH.h
--- a/Engine/Source/Developer/DeviceManager/Private/DeviceManagerPrivatePCH.h
+++ b/Engine/Source/Developer/DeviceManager/Private/DeviceManagerPrivatePCH.h
@@ -13,6 +13,7 @@
  *****************************************************************************/
 
 #include "ModuleManager.h"
+#include "TargetPlatform.h"
 
 
 /* Private includes
diff --git a/Engine/Source/Developer/LauncherServices/LauncherServices.Build.cs b/Engine/Source/Developer/LauncherServices/LauncherServices.Build.cs
--- a/Engine/Source/Developer/LauncherServices/LauncherServices.Build.cs
+++ b/Engine/Source/Developer/LauncherServices/LauncherServices.Build.cs
@@ -7,27 +7,34 @@ namespace UnrealBuildTool.Rules
 		public LauncherServices(TargetInfo Target)
 		{
 			PublicDependencyModuleNames.AddRange(
-			new string[]
-				{
+				new string[] {
 					"Core",
+				}
+			);
+
+			PublicIncludePathModuleNames.AddRange(
+				new string[] {
 					"TargetDeviceServices",
 				}
 			);
 
 			PrivateDependencyModuleNames.AddRange(
-				new string[]
-				{
+				new string[] {
 					"CoreUObject",
-					"Messaging",
 					"SessionMessages",
 					"TargetPlatform",
 					"UnrealEdMessages",
 				}
 			);
 
+			PrivateIncludePathModuleNames.AddRange(
+				new string[] {
+					"Messaging",
+				}
+			);
+
 			PrivateIncludePaths.AddRange(
-			new string[]
-				{
+				new string[] {
 					"Developer/LauncherServices/Private",
 					"Developer/LauncherServices/Private/Devices",
 					"Developer/LauncherServices/Private/Games",
diff --git a/Engine/Source/Developer/ScreenShotComparison/ScreenShotComparison.Build.cs b/Engine/Source/Developer/ScreenShotComparison/ScreenShotComparison.Build.cs
--- a/Engine/Source/Developer/ScreenShotComparison/ScreenShotComparison.Build.cs
+++ b/Engine/Source/Developer/ScreenShotComparison/ScreenShotComparison.Build.cs
@@ -6,25 +6,28 @@ public class ScreenShotComparison : ModuleRules
 {
 	public ScreenShotComparison(TargetInfo Target)
 	{
-		PrivateIncludePaths.AddRange(
-			new string[]
-			{
-				"Developer/ScreenShotComparison/Private",
-				"Developer/ScreenShotComparison/Private/Widgets",
-				"Developer/ScreenShotComparison/Private/Models",
-			}
-			);
-
 		PrivateDependencyModuleNames.AddRange(
-			new string[]
-			{
+			new string[] {
 				"Core",
                 "InputCore",
-				"SessionServices",
 				"Slate",
                 "EditorStyle",
 				"ScreenShotComparisonTools",
 			}
-			);
+		);
+
+		PrivateIncludePathModuleNames.AddRange(
+			new string[] {
+				"SessionServices",
+			}
+		);
+
+		PrivateIncludePaths.AddRange(
+			new string[] {
+				"Developer/ScreenShotComparison/Private",
+				"Developer/ScreenShotComparison/Private/Widgets",
+				"Developer/ScreenShotComparison/Private/Models",
+			}
+		);
 	}
 }
diff --git a/Engine/Source/Developer/ScreenShotComparisonTools/Private/ScreenShotComparisonToolsPrivatePCH.h b/Engine/Source/Developer/ScreenShotComparisonTools/Private/ScreenShotComparisonToolsPrivatePCH.h
--- a/Engine/Source/Developer/ScreenShotComparisonTools/Private/ScreenShotComparisonToolsPrivatePCH.h
+++ b/Engine/Source/Developer/ScreenShotComparisonTools/Private/ScreenShotComparisonToolsPrivatePCH.h
@@ -9,16 +9,20 @@
 
 #include "ScreenShotComparisonTools.h"
 
+
 /* Dependencies
  *****************************************************************************/
+
 #include "AutomationMessages.h"
+#include "Messaging.h"
 #include "ModuleManager.h"
 #include "Slate.h"
 
+
 /* Private includes
  *****************************************************************************/
+
 #include "ScreenShotBaseNode.h"
 #include "ScreenShotPlatformNode.h"
 #include "ScreenShotScreenNode.h"
-
 #include "ScreenShotManager.h"
diff --git a/Engine/Source/Developer/ScreenShotComparisonTools/Public/ScreenShotComparisonTools.h b/Engine/Source/Developer/ScreenShotComparisonTools/Public/ScreenShotComparisonTools.h
--- a/Engine/Source/Developer/ScreenShotComparisonTools/Public/ScreenShotComparisonTools.h
+++ b/Engine/Source/Developer/ScreenShotComparisonTools/Public/ScreenShotComparisonTools.h
@@ -12,7 +12,6 @@
 
 #include "Core.h"
 #include "ModuleManager.h"
-#include "Messaging.h"
 #include "IFilter.h"
 #include "TextFilter.h"
 #include "FilterCollection.h"
diff --git a/Engine/Source/Developer/ScreenShotComparisonTools/ScreenShotComparisonTools.Build.cs b/Engine/Source/Developer/ScreenShotComparisonTools/ScreenShotComparisonTools.Build.cs
--- a/Engine/Source/Developer/ScreenShotComparisonTools/ScreenShotComparisonTools.Build.cs
+++ b/Engine/Source/Developer/ScreenShotComparisonTools/ScreenShotComparisonTools.Build.cs
@@ -7,16 +7,13 @@ public class ScreenShotComparisonTools : ModuleRules
 	public ScreenShotComparisonTools(TargetInfo Target)
 	{
 		PublicDependencyModuleNames.AddRange(
-			new string[]
-			{
+			new string[] {
 				"Core",
-				"Messaging"
 			}
 		);
 
 		PrivateDependencyModuleNames.AddRange(
-			new string[]
-			{
+			new string[] {
 				"AutomationMessages",
 				"CoreUObject",
 				"UnrealEdMessages",
@@ -26,15 +23,13 @@ public class ScreenShotComparisonTools : ModuleRules
 		);
 
 		PrivateIncludePathModuleNames.AddRange(
-			new string[]
-			{
+			new string[] {
 				"Messaging"
 			}
 		);
 
 		PrivateIncludePaths.AddRange(
-			new string[]
-			{
+			new string[] {
 				"Developer/ScreenShotComparisonTools/Private"
 			}
 		);
diff --git a/Engine/Source/Developer/SessionFrontend/Private/SessionFrontendPrivatePCH.h b/Engine/Source/Developer/SessionFrontend/Private/SessionFrontendPrivatePCH.h
--- a/Engine/Source/Developer/SessionFrontend/Private/SessionFrontendPrivatePCH.h
+++ b/Engine/Source/Developer/SessionFrontend/Private/SessionFrontendPrivatePCH.h
@@ -6,18 +6,20 @@
 
 #pragma once
 
-
-#include "../Public/SessionFrontend.h"
+#include "SessionFrontend.h"
 
 
 /* Dependencies
  *****************************************************************************/
 
-#include "ModuleManager.h"
 #include "DesktopPlatformModule.h"
+#include "EditorStyle.h"
 #include "Json.h"
+#include "Messaging.h"
+#include "ModuleManager.h"
 #include "Slate.h"
-#include "EditorStyle.h"
+#include "TargetDeviceServices.h"
+#include "TargetPlatform.h"
 
 // @todo gmp: remove these dependencies by making the session front-end extensible
 #include "AutomationWindow.h"
diff --git a/Engine/Source/Developer/SessionFrontend/Public/SessionFrontend.h b/Engine/Source/Developer/SessionFrontend/Public/SessionFrontend.h
--- a/Engine/Source/Developer/SessionFrontend/Public/SessionFrontend.h
+++ b/Engine/Source/Developer/SessionFrontend/Public/SessionFrontend.h
@@ -14,7 +14,6 @@
 #include "ModuleInterface.h"
 #include "SessionServices.h"
 #include "Slate.h"
-#include "TargetDeviceServices.h"
 
 
 /* Interfaces
diff --git a/Engine/Source/Developer/SessionFrontend/SessionFrontend.Build.cs b/Engine/Source/Developer/SessionFrontend/SessionFrontend.Build.cs
--- a/Engine/Source/Developer/SessionFrontend/SessionFrontend.Build.cs
+++ b/Engine/Source/Developer/SessionFrontend/SessionFrontend.Build.cs
@@ -7,27 +7,37 @@ public class SessionFrontend : ModuleRules
 	public SessionFrontend(TargetInfo Target)
 	{
 		PublicDependencyModuleNames.AddRange(
-			new string[]
-			{
+			new string[] {
 				"Core",
-				"SessionServices",
 				"Slate",
                 "EditorStyle",
 			}
-		); 
-		
+		);
+
+		PublicIncludePathModuleNames.AddRange(
+			new string[] {
+				"SessionServices",
+			}
+		);
+
 		PrivateDependencyModuleNames.AddRange(
-			new string[]
-			{
+			new string[] {
 				"DesktopPlatform",
                 "InputCore",
-				"TargetDeviceServices",
 
 				// @todo gmp: remove these dependencies by making the session front-end extensible
 				"AutomationWindow",
 				"ScreenShotComparison",
 				"ScreenShotComparisonTools",
-				"Profiler"
+				"Profiler",
+				"TargetPlatform",
+			}
+		);
+
+		PrivateIncludePathModuleNames.AddRange(
+			new string[] {
+				"Messaging",
+				"TargetDeviceServices",
 			}
 		);
 
diff --git a/Engine/Source/Developer/SessionLauncher/Public/SessionLauncher.h b/Engine/Source/Developer/SessionLauncher/Public/SessionLauncher.h
--- a/Engine/Source/Developer/SessionLauncher/Public/SessionLauncher.h
+++ b/Engine/Source/Developer/SessionLauncher/Public/SessionLauncher.h
@@ -15,7 +15,6 @@
 #include "LauncherServices.h"
 #include "Slate.h"
 #include "EditorStyle.h"
-#include "TargetDeviceServices.h"
 
 
 /* Interfaces
diff --git a/Engine/Source/Developer/SessionLauncher/SessionLauncher.Build.cs b/Engine/Source/Developer/SessionLauncher/SessionLauncher.Build.cs
--- a/Engine/Source/Developer/SessionLauncher/SessionLauncher.Build.cs
+++ b/Engine/Source/Developer/SessionLauncher/SessionLauncher.Build.cs
@@ -6,21 +6,38 @@ public class SessionLauncher : ModuleRules
 {
 	public SessionLauncher(TargetInfo Target)
 	{
+		PublicDependencyModuleNames.AddRange(
+			new string[] {
+				"Core",
+				"CoreUObject",
+				"TargetPlatform",
+			}
+		);
+
+		PublicIncludePathModuleNames.AddRange(
+			new string[] {
+				"LauncherServices",
+			}
+		);
+
 		PrivateDependencyModuleNames.AddRange(
-			new string[]
-			{
+			new string[] {
 				"DesktopPlatform",
                 "InputCore",
-				"SessionServices",
 				"Slate",
                 "EditorStyle",
+			}
+		);
+
+		PrivateIncludePathModuleNames.AddRange(
+			new string[] {
+				"SessionServices",
 				"TargetDeviceServices",
 			}
 		);
 
 		PrivateIncludePaths.AddRange(
-			new string[]
-			{
+			new string[] {
 				"Developer/SessionLauncher/Private",
 				"Developer/SessionLauncher/Private/Models",
 				"Developer/SessionLauncher/Private/Widgets",
@@ -34,16 +51,5 @@ public class SessionLauncher : ModuleRules
 				"Developer/SessionLauncher/Private/Widgets/Toolbar",
 			}
 		);
-
-		PublicDependencyModuleNames.AddRange(
-			new string[]
-			{
-				"Core",
-				"CoreUObject",
-				"LauncherServices",
-				"SessionServices",
-				"TargetDeviceServices",
-			}
-		);
 	}
 }
diff --git a/Engine/Source/Developer/SessionServices/SessionServices.Build.cs b/Engine/Source/Developer/SessionServices/SessionServices.Build.cs
--- a/Engine/Source/Developer/SessionServices/SessionServices.Build.cs
+++ b/Engine/Source/Developer/SessionServices/SessionServices.Build.cs
@@ -18,9 +18,14 @@ public class SessionServices : ModuleRules
 			{
 				"CoreUObject",
 				"EngineMessages",
-				"Messaging",
 				"SessionMessages",
 			}
 		);
+
+		PrivateIncludePathModuleNames.AddRange(
+			new string[] {
+				"Messaging",
+			}
+		);
 	}
 }
diff --git a/Engine/Source/Developer/TargetDeviceServices/TargetDeviceServices.Build.cs b/Engine/Source/Developer/TargetDeviceServices/TargetDeviceServices.Build.cs
--- a/Engine/Source/Developer/TargetDeviceServices/TargetDeviceServices.Build.cs
+++ b/Engine/Source/Developer/TargetDeviceServices/TargetDeviceServices.Build.cs
@@ -18,6 +18,11 @@ namespace UnrealBuildTool.Rules
 				new string[]
 				{
 					"CoreUObject",
+				}
+			);
+
+			PrivateIncludePathModuleNames.AddRange(
+				new string[] {
 					"Messaging",
 				}
 			);
diff --git a/Engine/Source/Programs/UnrealFrontend/UnrealFrontend.Build.cs b/Engine/Source/Programs/UnrealFrontend/UnrealFrontend.Build.cs
--- a/Engine/Source/Programs/UnrealFrontend/UnrealFrontend.Build.cs
+++ b/Engine/Source/Programs/UnrealFrontend/UnrealFrontend.Build.cs
@@ -20,7 +20,6 @@ public class UnrealFrontend : ModuleRules
 				"AutomationController",
 				"Core",
 				"DeviceManager",
-				"Messaging",
 				"ProfilerClient",
 				"Projects",
 				"SessionFrontend",
@@ -39,5 +38,11 @@ public class UnrealFrontend : ModuleRules
 				"UdpMessaging",
 			}
 		);
+
+		PrivateIncludePathModuleNames.AddRange(
+			new string[] {
+				"Messaging",
+			}
+		);
 	}
 }
\ No newline at end of file
