DeviceManager: documentation cleanup pass

diff --git a/Engine/Source/Developer/DeviceManager/Private/Models/DeviceBrowserFilter.h b/Engine/Source/Developer/DeviceManager/Private/Models/DeviceBrowserFilter.h
--- a/Engine/Source/Developer/DeviceManager/Private/Models/DeviceBrowserFilter.h
+++ b/Engine/Source/Developer/DeviceManager/Private/Models/DeviceBrowserFilter.h
@@ -7,14 +7,10 @@
 #pragma once
 
 
-/**
- * Type definition for shared pointers to instances of FDeviceBrowserFilter.
- */
+/** Type definition for shared pointers to instances of FDeviceBrowserFilter. */
 typedef TSharedPtr<class FDeviceBrowserFilter> FDeviceBrowserFilterPtr;
 
-/**
- * Type definition for shared references to instances of FDeviceBrowserFilter.
- */
+/** Type definition for shared references to instances of FDeviceBrowserFilter. */
 typedef TSharedRef<class FDeviceBrowserFilter> FDeviceBrowserFilterRef;
 
 
@@ -28,8 +24,7 @@ public:
 	/**
 	 * Filters the specified target device service based on the current filter settings.
 	 *
-	 * @param Device - The service to filter.
-	 *
+	 * @param Device The service to filter.
 	 * @return true if the service passed the filter, false otherwise.
 	 */
 	bool FilterDeviceService( const ITargetDeviceServicePtr& DeviceService )
@@ -55,8 +50,7 @@ public:
 	/**
 	 * Gets the number of target device services that have the specified platform.
 	 *
-	 * @param PlatformName - The name of the platform.
-	 *
+	 * @param PlatformName The name of the platform.
 	 * @return The number of matching services.
 	 */
 	int32 GetServiceCountPerPlatform( const FString& PlatformName ) const
@@ -87,8 +81,7 @@ public:
 	/**
 	 * Checks whether the specified platform is enabled in the filter.
 	 *
-	 * @param PlatformName - The name of the platform to check.
-	 *
+	 * @param PlatformName The name of the platform to check.
 	 * @return true if the platform is enabled, false otherwise.
 	 */
 	bool IsPlatformEnabled( const FString& PlatformName ) const
@@ -99,7 +92,7 @@ public:
 	/**
 	 * Populates the filter from the given list of target device proxies.
 	 *
-	 * @param DeviceProxies - The list of device proxies to populate the filter from.
+	 * @param DeviceProxies The list of device proxies to populate the filter from.
 	 */
 	void ResetFilter( const TArray<ITargetDeviceServicePtr>& DeviceServices )
 	{
@@ -128,7 +121,7 @@ public:
 	/**
 	 * Sets the current device search string.
 	 *
-	 * @param SearchText - The search string.
+	 * @param SearchText The search string.
 	 */
 	void SetDeviceSearchString( const FText& SearchText )
 	{
@@ -143,8 +136,8 @@ public:
 	/**
 	 * Sets the enabled state of the specified device proxy platform.
 	 *
-	 * @param PlatformName - The name of the platform to enable or disable in the filter.
-	 * @param Enabled - Whether the platform is enabled.
+	 * @param PlatformName The name of the platform to enable or disable in the filter.
+	 * @param Enabled Whether the platform is enabled.
 	 */
 	void SetPlatformEnabled( const FString& PlatformName, bool Enabled )
 	{
@@ -199,7 +192,7 @@ private:
 	TMap<FString, int32> PlatformCounters;
 
 	// Holds the list of platform filters.
-	TArray<TSharedPtr<FString> > PlatformList;
+	TArray<TSharedPtr<FString>> PlatformList;
 
 private:
 
diff --git a/Engine/Source/Developer/DeviceManager/Private/Models/DeviceDetailsFeature.h b/Engine/Source/Developer/DeviceManager/Private/Models/DeviceDetailsFeature.h
--- a/Engine/Source/Developer/DeviceManager/Private/Models/DeviceDetailsFeature.h
+++ b/Engine/Source/Developer/DeviceManager/Private/Models/DeviceDetailsFeature.h
@@ -7,14 +7,10 @@
 #pragma once
 
 
-/**
- * Type definition for shared pointers to instances of FDeviceDetailsFeature.
- */
+/** Type definition for shared pointers to instances of FDeviceDetailsFeature. */
 typedef TSharedPtr<struct FDeviceDetailsFeature> FDeviceDetailsFeaturePtr;
 
-/**
- * Type definition for shared references to instances of FDeviceDetailsFeature.
- */
+/** Type definition for shared references to instances of FDeviceDetailsFeature. */
 typedef TSharedRef<struct FDeviceDetailsFeature> FDeviceDetailsFeatureRef;
 
 
@@ -23,21 +19,17 @@ typedef TSharedRef<struct FDeviceDetailsFeature> FDeviceDetailsFeatureRef;
  */
 struct FDeviceDetailsFeature
 {
-	/**
-	 * Holds the name of the feature.
-	 */
+	/** Holds the name of the feature. */
 	FString FeatureName;
 
-	/**
-	 * Whether the feature is available.
-	 */
+	/** Whether the feature is available. */
 	bool Available;
 
 	/**
 	 * Creates and initializes a new instance.
 	 *
-	 * @param InFeatureName - The name of the feature.
-	 * @param InAvailable - Whether the feature is available.
+	 * @param InFeatureName The name of the feature.
+	 * @param InAvailable Whether the feature is available.
 	 */
 	FDeviceDetailsFeature( const FString& InFeatureName, bool InAvailable )
 		: FeatureName(InFeatureName)
diff --git a/Engine/Source/Developer/DeviceManager/Private/Models/DeviceManagerModel.h b/Engine/Source/Developer/DeviceManager/Private/Models/DeviceManagerModel.h
--- a/Engine/Source/Developer/DeviceManager/Private/Models/DeviceManagerModel.h
+++ b/Engine/Source/Developer/DeviceManager/Private/Models/DeviceManagerModel.h
@@ -7,14 +7,10 @@
 #pragma once
 
 
-/**
- * Type definition for shared pointers to instances of FDeviceManagerModel.
- */
+/** Type definition for shared pointers to instances of FDeviceManagerModel. */
 typedef TSharedPtr<class FDeviceManagerModel> FDeviceManagerModelPtr;
 
-/**
- * Type definition for shared references to instances of FDeviceManagerModel.
- */
+/** Type definition for shared references to instances of FDeviceManagerModel. */
 typedef TSharedRef<class FDeviceManagerModel> FDeviceManagerModelRef;
 
 
@@ -38,7 +34,7 @@ public:
 	/**
 	 * Selects the specified device service (or none if nullptr).
 	 *
-	 * @param DeviceService - The device service to select.
+	 * @param DeviceService The device service to select.
 	 */
 	void SelectDeviceService( const ITargetDeviceServicePtr& DeviceService )
 	{
diff --git a/Engine/Source/Developer/DeviceManager/Private/Widgets/Apps/SDeviceApps.h b/Engine/Source/Developer/DeviceManager/Private/Widgets/Apps/SDeviceApps.h
--- a/Engine/Source/Developer/DeviceManager/Private/Widgets/Apps/SDeviceApps.h
+++ b/Engine/Source/Developer/DeviceManager/Private/Widgets/Apps/SDeviceApps.h
@@ -30,8 +30,8 @@ public:
 	/**
 	 * Constructs the widget.
 	 *
-	 * @param InArgs - The construction arguments.
-	 * @param InModel - The view model to use.
+	 * @param InArgs The construction arguments.
+	 * @param InModel The view model to use.
 	 */
 	void Construct( const FArguments& InArgs, const FDeviceManagerModelRef& InModel );
 
@@ -58,7 +58,7 @@ private:
 	TArray<TSharedPtr<FString> > AppList;
 
 	// Holds the application list view.
-	TSharedPtr<SListView<TSharedPtr<FString> > > AppListView;
+	TSharedPtr<SListView<TSharedPtr<FString>>> AppListView;
 
 	// Holds a pointer the device manager's view model.
 	FDeviceManagerModelPtr Model;
diff --git a/Engine/Source/Developer/DeviceManager/Private/Widgets/Apps/SDeviceAppsAppListRow.h b/Engine/Source/Developer/DeviceManager/Private/Widgets/Apps/SDeviceAppsAppListRow.h
--- a/Engine/Source/Developer/DeviceManager/Private/Widgets/Apps/SDeviceAppsAppListRow.h
+++ b/Engine/Source/Developer/DeviceManager/Private/Widgets/Apps/SDeviceAppsAppListRow.h
@@ -26,7 +26,7 @@ public:
 	/**
 	 * Constructs the widget.
 	 *
-	 * @param InArgs - The arguments.
+	 * @param InArgs The construction arguments.
 	 */
 	void Construct( const FArguments& InArgs, const TSharedRef<STableViewBase>& InOwnerTableView )
 	{
@@ -38,8 +38,7 @@ public:
 	/**
 	 * Generates the widget for the specified column.
 	 *
-	 * @param ColumnName - The name of the column to generate the widget for.
-	 *
+	 * @param ColumnName The name of the column to generate the widget for.
 	 * @return The widget.
 	 */
 	virtual TSharedRef<SWidget> GenerateWidgetForColumn( const FName& ColumnName ) OVERRIDE
diff --git a/Engine/Source/Developer/DeviceManager/Private/Widgets/Browser/SDeviceBrowser.h b/Engine/Source/Developer/DeviceManager/Private/Widgets/Browser/SDeviceBrowser.h
--- a/Engine/Source/Developer/DeviceManager/Private/Widgets/Browser/SDeviceBrowser.h
+++ b/Engine/Source/Developer/DeviceManager/Private/Widgets/Browser/SDeviceBrowser.h
@@ -31,10 +31,10 @@ public:
 	/**
 	 * Constructs the widget.
 	 *
-	 * @param InArgs - The construction arguments.
-	 * @param InModel - The view model to use.
-	 * @param InDeviceServiceManager - The target device service manager to use.
-	 * @param InUICommandList - The UI command list to use.
+	 * @param InArgs The construction arguments.
+	 * @param InModel The view model to use.
+	 * @param InDeviceServiceManager The target device service manager to use.
+	 * @param InUICommandList The UI command list to use.
 	 */
 	void Construct( const FArguments& InArgs, const FDeviceManagerModelRef& InModel, const ITargetDeviceServiceManagerRef& InDeviceServiceManager, const TSharedPtr<FUICommandList>& InUICommandList );
 
@@ -51,7 +51,7 @@ protected:
 	/**
 	 * Reloads the list of target device services.
 	 *
-	 * @param FullyReload - Whether to fully reload the service entries, or only re-apply filtering.
+	 * @param FullyReload Whether to fully reload the service entries, or only re-apply filtering.
 	 */
 	void ReloadDeviceServiceList( bool FullyReload );
 
diff --git a/Engine/Source/Developer/DeviceManager/Private/Widgets/Browser/SDeviceBrowserContextMenu.h b/Engine/Source/Developer/DeviceManager/Private/Widgets/Browser/SDeviceBrowserContextMenu.h
--- a/Engine/Source/Developer/DeviceManager/Private/Widgets/Browser/SDeviceBrowserContextMenu.h
+++ b/Engine/Source/Developer/DeviceManager/Private/Widgets/Browser/SDeviceBrowserContextMenu.h
@@ -26,8 +26,8 @@ public:
 	/**
 	 * Construct this widget.
 	 *
-	 * @param InArgs - The declaration data for this widget.
-	 * @param InUICommandList - The UI command list to use.
+	 * @param InArgs The declaration data for this widget.
+	 * @param InUICommandList The UI command list to use.
 	 */
 	void Construct( const FArguments& InArgs, const TSharedPtr<FUICommandList>& InUICommandList )
 	{
diff --git a/Engine/Source/Developer/DeviceManager/Private/Widgets/Browser/SDeviceBrowserDeviceAdder.h b/Engine/Source/Developer/DeviceManager/Private/Widgets/Browser/SDeviceBrowserDeviceAdder.h
--- a/Engine/Source/Developer/DeviceManager/Private/Widgets/Browser/SDeviceBrowserDeviceAdder.h
+++ b/Engine/Source/Developer/DeviceManager/Private/Widgets/Browser/SDeviceBrowserDeviceAdder.h
@@ -23,8 +23,8 @@ public:
 	/**
 	 * Constructs the widget.
 	 *
-	 * @param InArgs - The construction arguments.
-	 * @param InDeviceManager - The target device manager to use.
+	 * @param InArgs The construction arguments.
+	 * @param InDeviceManager The target device manager to use.
 	 */
 	void Construct( const FArguments& InArgs, const ITargetDeviceServiceManagerRef& InDeviceServiceManager );
 
@@ -87,8 +87,8 @@ private:
 	TSharedPtr<SOverlay> UserDataOverlay;
 
 	// Holds the platforms combo box.
-	TSharedPtr<SComboBox<TSharedPtr<FString> > > PlatformComboBox;
+	TSharedPtr<SComboBox<TSharedPtr<FString>>> PlatformComboBox;
 
 	// Holds the list of known platforms.
-	TArray<TSharedPtr<FString> > PlatformList;
+	TArray<TSharedPtr<FString>> PlatformList;
 };
diff --git a/Engine/Source/Developer/DeviceManager/Private/Widgets/Browser/SDeviceBrowserDeviceListRow.h b/Engine/Source/Developer/DeviceManager/Private/Widgets/Browser/SDeviceBrowserDeviceListRow.h
--- a/Engine/Source/Developer/DeviceManager/Private/Widgets/Browser/SDeviceBrowserDeviceListRow.h
+++ b/Engine/Source/Developer/DeviceManager/Private/Widgets/Browser/SDeviceBrowserDeviceListRow.h
@@ -28,7 +28,7 @@ public:
 	/**
 	 * Constructs the widget.
 	 *
-	 * @param InArgs - The arguments.
+	 * @param InArgs The arguments.
 	 */
 	void Construct( const FArguments& InArgs, const TSharedRef<STableViewBase>& InOwnerTableView )
 	{
@@ -46,8 +46,7 @@ public:
 	/**
 	 * Generates the widget for the specified column.
 	 *
-	 * @param ColumnName - The name of the column to generate the widget for.
-	 *
+	 * @param ColumnName The name of the column to generate the widget for.
 	 * @return The widget.
 	 */
 	BEGIN_SLATE_FUNCTION_BUILD_OPTIMIZATION
diff --git a/Engine/Source/Developer/DeviceManager/Private/Widgets/Browser/SDeviceBrowserFilterBar.h b/Engine/Source/Developer/DeviceManager/Private/Widgets/Browser/SDeviceBrowserFilterBar.h
--- a/Engine/Source/Developer/DeviceManager/Private/Widgets/Browser/SDeviceBrowserFilterBar.h
+++ b/Engine/Source/Developer/DeviceManager/Private/Widgets/Browser/SDeviceBrowserFilterBar.h
@@ -37,8 +37,8 @@ public:
 	/**
 	 * Construct this widget
 	 *
-	 * @param InArgs - The declaration data for this widget.
-	 * @param InFilter - The filter model to use.
+	 * @param InArgs The declaration data for this widget.
+	 * @param InFilter The filter model to use.
 	 */
 	void Construct( const FArguments& InArgs, FDeviceBrowserFilterRef InFilter );
 
@@ -71,5 +71,5 @@ private:
 	TSharedPtr<SSearchBox> FilterStringTextBox;
 
 	// Holds the platform filters list view.
-	TSharedPtr<SListView<TSharedPtr<FString> > > PlatformListView;
+	TSharedPtr<SListView<TSharedPtr<FString>>> PlatformListView;
 };
diff --git a/Engine/Source/Developer/DeviceManager/Private/Widgets/Browser/SDeviceBrowserTooltip.h b/Engine/Source/Developer/DeviceManager/Private/Widgets/Browser/SDeviceBrowserTooltip.h
--- a/Engine/Source/Developer/DeviceManager/Private/Widgets/Browser/SDeviceBrowserTooltip.h
+++ b/Engine/Source/Developer/DeviceManager/Private/Widgets/Browser/SDeviceBrowserTooltip.h
@@ -26,9 +26,9 @@ public:
 	/**
 	 * Constructs the widget.
 	 *
-	 * @param InArgs - The construction arguments.
-	 * @param InDeviceServiceManager - The target device service manager to use.
-	 * @param InDeviceManagerState - The optional device manager view state.
+	 * @param InArgs The construction arguments.
+	 * @param InDeviceServiceManager The target device service manager to use.
+	 * @param InDeviceManagerState The optional device manager view state.
 	 */
 	void Construct( const FArguments& InArgs, const ITargetDeviceServiceRef& InDeviceService )
 	{
diff --git a/Engine/Source/Developer/DeviceManager/Private/Widgets/Details/SDeviceDetails.h b/Engine/Source/Developer/DeviceManager/Private/Widgets/Details/SDeviceDetails.h
--- a/Engine/Source/Developer/DeviceManager/Private/Widgets/Details/SDeviceDetails.h
+++ b/Engine/Source/Developer/DeviceManager/Private/Widgets/Details/SDeviceDetails.h
@@ -30,8 +30,8 @@ public:
 	/**
 	 * Constructs the widget.
 	 *
-	 * @param InArgs - The construction arguments.
-	 * @param InModel - The view model to use.
+	 * @param InArgs The construction arguments.
+	 * @param InModel The view model to use.
 	 */
 	void Construct( const FArguments& InArgs, const FDeviceManagerModelRef& InModel );
 
diff --git a/Engine/Source/Developer/DeviceManager/Private/Widgets/Details/SDeviceDetailsFeatureListRow.h b/Engine/Source/Developer/DeviceManager/Private/Widgets/Details/SDeviceDetailsFeatureListRow.h
--- a/Engine/Source/Developer/DeviceManager/Private/Widgets/Details/SDeviceDetailsFeatureListRow.h
+++ b/Engine/Source/Developer/DeviceManager/Private/Widgets/Details/SDeviceDetailsFeatureListRow.h
@@ -27,8 +27,8 @@ public:
 	/**
 	 * Constructs the widget.
 	 *
-	 * @param InArgs - The construction arguments.
-	 * @param InOwnerTableView - The table view that owns this row.
+	 * @param InArgs The construction arguments.
+	 * @param InOwnerTableView The table view that owns this row.
 	 */
 	void Construct( const FArguments& InArgs, const TSharedRef<STableViewBase>& InOwnerTableView, const FDeviceDetailsFeatureRef& InFeature )
 	{
diff --git a/Engine/Source/Developer/DeviceManager/Private/Widgets/Processes/SDeviceProcesses.h b/Engine/Source/Developer/DeviceManager/Private/Widgets/Processes/SDeviceProcesses.h
--- a/Engine/Source/Developer/DeviceManager/Private/Widgets/Processes/SDeviceProcesses.h
+++ b/Engine/Source/Developer/DeviceManager/Private/Widgets/Processes/SDeviceProcesses.h
@@ -30,8 +30,8 @@ public:
 	/**
 	 * Constructs the widget.
 	 *
-	 * @param InArgs - The construction arguments.
-	 * @param InModel - The view model to use.
+	 * @param InArgs The construction arguments.
+	 * @param InModel The view model to use.
 	 */
 	void Construct( const FArguments& InArgs, const FDeviceManagerModelRef& InModel );
 
diff --git a/Engine/Source/Developer/DeviceManager/Private/Widgets/Processes/SDeviceProcessesProcessListRow.h b/Engine/Source/Developer/DeviceManager/Private/Widgets/Processes/SDeviceProcessesProcessListRow.h
--- a/Engine/Source/Developer/DeviceManager/Private/Widgets/Processes/SDeviceProcessesProcessListRow.h
+++ b/Engine/Source/Developer/DeviceManager/Private/Widgets/Processes/SDeviceProcessesProcessListRow.h
@@ -27,7 +27,7 @@ public:
 	/**
 	 * Constructs the widget.
 	 *
-	 * @param InArgs - The arguments.
+	 * @param InArgs The arguments.
 	 */
 	void Construct( const FArguments& InArgs, const TSharedRef<STableViewBase>& InOwnerTableView )
 	{
@@ -41,7 +41,7 @@ public:
 	/**
 	 * Generates the widget for the specified column.
 	 *
-	 * @param ColumnName - The name of the column to generate the widget for.
+	 * @param ColumnName The name of the column to generate the widget for.
 	 *
 	 * @return The widget.
 	 */
diff --git a/Engine/Source/Developer/DeviceManager/Private/Widgets/Processes/SDeviceProcessesProcessTreeNode.h b/Engine/Source/Developer/DeviceManager/Private/Widgets/Processes/SDeviceProcessesProcessTreeNode.h
--- a/Engine/Source/Developer/DeviceManager/Private/Widgets/Processes/SDeviceProcessesProcessTreeNode.h
+++ b/Engine/Source/Developer/DeviceManager/Private/Widgets/Processes/SDeviceProcessesProcessTreeNode.h
@@ -7,14 +7,10 @@
 #pragma once
 
 
-/**
- * Type definition for shared pointers to instances of FDeviceProcessesProcessTreeNode.
- */
+/** Type definition for shared pointers to instances of FDeviceProcessesProcessTreeNode. */
 typedef TSharedPtr<class FDeviceProcessesProcessTreeNode> FDeviceProcessesProcessTreeNodePtr;
 
-/**
- * Type definition for shared references to instances of FDeviceProcessesProcessTreeNode.
- */
+/** Type definition for shared references to instances of FDeviceProcessesProcessTreeNode. */
 typedef TSharedRef<class FDeviceProcessesProcessTreeNode> FDeviceProcessesProcessTreeNodeRef;
 
 
@@ -28,7 +24,7 @@ public:
 	/**
 	 * Creates and initializes a new instance.
 	 *
-	 * @param InProcessInfo - The node's process information.
+	 * @param InProcessInfo The node's process information.
 	 */
 	FDeviceProcessesProcessTreeNode( const FTargetDeviceProcessInfo& InProcessInfo )
 		: ProcessInfo(InProcessInfo)
@@ -87,7 +83,7 @@ public:
 	/**
 	 * Sets the parent node.
 	 *
-	 * @param Node - The parent node to set.
+	 * @param Node The parent node to set.
 	 */
 	void SetParent( const FDeviceProcessesProcessTreeNodePtr& Node )
 	{
diff --git a/Engine/Source/Developer/DeviceManager/Private/Widgets/SDeviceManager.h b/Engine/Source/Developer/DeviceManager/Private/Widgets/SDeviceManager.h
--- a/Engine/Source/Developer/DeviceManager/Private/Widgets/SDeviceManager.h
+++ b/Engine/Source/Developer/DeviceManager/Private/Widgets/SDeviceManager.h
@@ -30,10 +30,10 @@ public:
 	/**
 	 * Constructs the widget.
 	 *
-	 * @param InArgs - The Slate argument list.
-	 * @param InDeviceServiceManager - The target device manager to use.
-	 * @param ConstructUnderMajorTab - The major tab which will contain the session front-end.
-	 * @param ConstructUnderWindow - The window in which this widget is being constructed.
+	 * @param InArgs The Slate argument list.
+	 * @param InDeviceServiceManager The target device manager to use.
+	 * @param ConstructUnderMajorTab The major tab which will contain the session front-end.
+	 * @param ConstructUnderWindow The window in which this widget is being constructed.
 	 */
 	void Construct( const FArguments& InArgs, const ITargetDeviceServiceManagerRef& InDeviceServiceManager, const TSharedRef<SDockTab>& ConstructUnderMajorTab, const TSharedPtr<SWindow>& ConstructUnderWindow );
 
@@ -47,18 +47,17 @@ protected:
 	/**
 	 * Fills the Window menu with menu items.
 	 *
-	 * @param MenuBuilder - The multi-box builder that should be filled with content for this pull-down menu.
-	 * @param RootMenuGroup - The root menu group.
-	 * @param AppMenuGroup - The application menu group.
-	 * @param TabManager - A Tab Manager from which to populate tab spawner menu items.
+	 * @param MenuBuilder The multi-box builder that should be filled with content for this pull-down menu.
+	 * @param RootMenuGroup The root menu group.
+	 * @param AppMenuGroup The application menu group.
+	 * @param TabManager A Tab Manager from which to populate tab spawner menu items.
 	 */
 	static void FillWindowMenu( FMenuBuilder& MenuBuilder, TSharedRef<FWorkspaceItem> RootMenuGroup, TSharedRef<FWorkspaceItem> AppMenuGroup, const TSharedPtr<FTabManager> TabManager );
 
 	/**
 	 * Validates actions on the specified device.
 	 *
-	 * @param Device - The device to perform an action on.
-	 *
+	 * @param Device The device to perform an action on.
 	 * @return true if actions on the device are permitted, false otherwise.
 	 */
 	bool ValidateDeviceAction( const ITargetDeviceRef& Device ) const;
diff --git a/Engine/Source/Developer/DeviceManager/Private/Widgets/Shared/SDeviceQuickInfo.h b/Engine/Source/Developer/DeviceManager/Private/Widgets/Shared/SDeviceQuickInfo.h
--- a/Engine/Source/Developer/DeviceManager/Private/Widgets/Shared/SDeviceQuickInfo.h
+++ b/Engine/Source/Developer/DeviceManager/Private/Widgets/Shared/SDeviceQuickInfo.h
@@ -30,9 +30,9 @@ public:
 	/**
 	 * Constructs the widget.
 	 *
-	 * @param InArgs - The construction arguments.
-	 * @param InDeviceServiceManager - The target device service manager to use.
-	 * @param InDeviceManagerState - The optional device manager view state.
+	 * @param InArgs The construction arguments.
+	 * @param InDeviceServiceManager The target device service manager to use.
+	 * @param InDeviceManagerState The optional device manager view state.
 	 */
 	void Construct( const FArguments& InArgs )
 	{
@@ -165,7 +165,7 @@ public:
 	/**
 	 * Sets the device service whose information is being shown.
 	 *
-	 * @param InDeviceService - The device service to show.
+	 * @param InDeviceService The device service to show.
 	 */
 	void SetDeviceService( const ITargetDeviceServicePtr& InDeviceService )
 	{
diff --git a/Engine/Source/Developer/DeviceManager/Private/Widgets/Toolbar/SDeviceToolbar.h b/Engine/Source/Developer/DeviceManager/Private/Widgets/Toolbar/SDeviceToolbar.h
--- a/Engine/Source/Developer/DeviceManager/Private/Widgets/Toolbar/SDeviceToolbar.h
+++ b/Engine/Source/Developer/DeviceManager/Private/Widgets/Toolbar/SDeviceToolbar.h
@@ -23,9 +23,9 @@ public:
 	/**
 	 * Constructs the widget.
 	 *
-	 * @param InArgs - The construction arguments.
-	 * @param InModel - The view model to use.
-	 * @param InUICommandList - The UI command list to use.
+	 * @param InArgs The construction arguments.
+	 * @param InModel The view model to use.
+	 * @param InUICommandList The UI command list to use.
 	 */
 	void Construct( const FArguments& InArgs, const FDeviceManagerModelRef& InModel, const TSharedPtr<FUICommandList>& InUICommandList );
 
diff --git a/Engine/Source/Developer/DeviceManager/Public/Interfaces/IDeviceManagerModule.h b/Engine/Source/Developer/DeviceManager/Public/Interfaces/IDeviceManagerModule.h
--- a/Engine/Source/Developer/DeviceManager/Public/Interfaces/IDeviceManagerModule.h
+++ b/Engine/Source/Developer/DeviceManager/Public/Interfaces/IDeviceManagerModule.h
@@ -18,8 +18,7 @@ public:
 	/**
 	 * Creates a device manager widget.
 	 *
-	 * @param DeviceServiceManager - The target device manager to use.
-	 *
+	 * @param DeviceServiceManager The target device manager to use.
 	 * @return The new widget.
 	 */
 	virtual TSharedRef<class SWidget> CreateDeviceManager( const ITargetDeviceServiceManagerRef& DeviceServiceManager, const TSharedRef<SDockTab>& ConstructUnderMajorTab, const TSharedPtr<SWindow>& ConstructUnderWindow ) = 0;
