Removed some check() calls when retrieving PlatformInfo for a device

TTP# 338166 - CRITICAL: Regression: UFE: CRASH: Deleting a profile under the Advanced section

Some places are using this only to get the icon information, so they now use an empty brush should the platform information not be found (it shouldn't really happen, but the device information may sometimes be missing a valid target platform name).

Also made sure that any missing icons wouldn't cause list views to misalign their content.

diff --git a/Engine/Source/Developer/DeviceManager/Private/Widgets/Browser/SDeviceBrowserDeviceAdder.cpp b/Engine/Source/Developer/DeviceManager/Private/Widgets/Browser/SDeviceBrowserDeviceAdder.cpp
--- a/Engine/Source/Developer/DeviceManager/Private/Widgets/Browser/SDeviceBrowserDeviceAdder.cpp
+++ b/Engine/Source/Developer/DeviceManager/Private/Widgets/Browser/SDeviceBrowserDeviceAdder.cpp
@@ -300,7 +300,6 @@ FString SDeviceBrowserDeviceAdder::HandlePlatformComboBoxContentText( ) const
 TSharedRef<SWidget> SDeviceBrowserDeviceAdder::HandlePlatformComboBoxGenerateWidget( TSharedPtr<FString> Item )
 {
 	const PlatformInfo::FPlatformInfo* const PlatformInfo = PlatformInfo::FindPlatformInfo(**Item);
-	check(PlatformInfo);
 
 	return
 		SNew(SHorizontalBox)
@@ -309,8 +308,13 @@ TSharedRef<SWidget> SDeviceBrowserDeviceAdder::HandlePlatformComboBoxGenerateWid
 			.AutoWidth()
 			.HAlign(HAlign_Left)
 			[
-				SNew(SImage)
-					.Image(FEditorStyle::GetBrush(PlatformInfo->GetIconStyleName(PlatformInfo::EPlatformIconSize::Normal)))
+				SNew(SBox)
+					.WidthOverride(24)
+					.HeightOverride(24)
+					[
+						SNew(SImage)
+							.Image((PlatformInfo) ? FEditorStyle::GetBrush(PlatformInfo->GetIconStyleName(PlatformInfo::EPlatformIconSize::Normal)) : FStyleDefaults::GetNoBrush())
+					]
 			]
 
 		+ SHorizontalBox::Slot()
diff --git a/Engine/Source/Developer/DeviceManager/Private/Widgets/Browser/SDeviceBrowserDeviceListRow.h b/Engine/Source/Developer/DeviceManager/Private/Widgets/Browser/SDeviceBrowserDeviceListRow.h
--- a/Engine/Source/Developer/DeviceManager/Private/Widgets/Browser/SDeviceBrowserDeviceListRow.h
+++ b/Engine/Source/Developer/DeviceManager/Private/Widgets/Browser/SDeviceBrowserDeviceListRow.h
@@ -62,13 +62,14 @@ public:
 		else if (ColumnName == TEXT("Icon"))
 		{
 			const PlatformInfo::FPlatformInfo* const PlatformInfo = PlatformInfo::FindPlatformInfo(*DeviceService->GetDeviceId().GetPlatformName());
-			check(PlatformInfo);
 
 			return SNew(SBox)
 				.Padding(FMargin(4.0f, 0.0f))
+				.WidthOverride(24)
+				.HeightOverride(24)
 				[
 					SNew(SImage)
-						.Image(FEditorStyle::GetBrush(PlatformInfo->GetIconStyleName(PlatformInfo::EPlatformIconSize::Normal)))
+						.Image((PlatformInfo) ? FEditorStyle::GetBrush(PlatformInfo->GetIconStyleName(PlatformInfo::EPlatformIconSize::Normal)) : FStyleDefaults::GetNoBrush())
 				];
 		}
 		else if (ColumnName == TEXT("Name"))
diff --git a/Engine/Source/Developer/DeviceManager/Private/Widgets/Browser/SDeviceBrowserFilterBar.cpp b/Engine/Source/Developer/DeviceManager/Private/Widgets/Browser/SDeviceBrowserFilterBar.cpp
--- a/Engine/Source/Developer/DeviceManager/Private/Widgets/Browser/SDeviceBrowserFilterBar.cpp
+++ b/Engine/Source/Developer/DeviceManager/Private/Widgets/Browser/SDeviceBrowserFilterBar.cpp
@@ -110,7 +110,6 @@ BEGIN_SLATE_FUNCTION_BUILD_OPTIMIZATION
 TSharedRef<ITableRow> SDeviceBrowserFilterBar::HandlePlatformListViewGenerateRow( TSharedPtr<FString> PlatformName, const TSharedRef<STableViewBase>& OwnerTable )
 {
 	const PlatformInfo::FPlatformInfo* const PlatformInfo = PlatformInfo::FindPlatformInfo(**PlatformName);
-	check(PlatformInfo);
 
 	return SNew(STableRow<TSharedPtr<FString> >, OwnerTable)
 		.Content()
@@ -126,8 +125,13 @@ TSharedRef<ITableRow> SDeviceBrowserFilterBar::HandlePlatformListViewGenerateRow
 					+ SHorizontalBox::Slot()
 						.AutoWidth()
 						[
-							SNew(SImage)
-								.Image(FEditorStyle::GetBrush(PlatformInfo->GetIconStyleName(PlatformInfo::EPlatformIconSize::Normal)))
+							SNew(SBox)
+								.WidthOverride(24)
+								.HeightOverride(24)
+								[
+									SNew(SImage)
+										.Image((PlatformInfo) ? FEditorStyle::GetBrush(PlatformInfo->GetIconStyleName(PlatformInfo::EPlatformIconSize::Normal)) : FStyleDefaults::GetNoBrush())
+								]
 						]
 
 					+ SHorizontalBox::Slot()
diff --git a/Engine/Source/Developer/DeviceManager/Private/Widgets/Shared/SDeviceQuickInfo.h b/Engine/Source/Developer/DeviceManager/Private/Widgets/Shared/SDeviceQuickInfo.h
--- a/Engine/Source/Developer/DeviceManager/Private/Widgets/Shared/SDeviceQuickInfo.h
+++ b/Engine/Source/Developer/DeviceManager/Private/Widgets/Shared/SDeviceQuickInfo.h
@@ -242,12 +242,13 @@ private:
 		if (DeviceService.IsValid())
 		{
 			const PlatformInfo::FPlatformInfo* const PlatformInfo = PlatformInfo::FindPlatformInfo(*DeviceService->GetDeviceId().GetPlatformName());
-			check(PlatformInfo);
-
-			return FEditorStyle::GetBrush(PlatformInfo->GetIconStyleName(PlatformInfo::EPlatformIconSize::XLarge));
+			if(PlatformInfo)
+			{
+				return FEditorStyle::GetBrush(PlatformInfo->GetIconStyleName(PlatformInfo::EPlatformIconSize::XLarge));
+			}
 		}
 
-		return NULL;
+		return FStyleDefaults::GetNoBrush();
 	}
 
 	// Callback for getting the name of the device's platform.
diff --git a/Engine/Source/Developer/SessionFrontend/Private/Widgets/Browser/SSessionBrowserInstanceListRow.h b/Engine/Source/Developer/SessionFrontend/Private/Widgets/Browser/SSessionBrowserInstanceListRow.h
--- a/Engine/Source/Developer/SessionFrontend/Private/Widgets/Browser/SSessionBrowserInstanceListRow.h
+++ b/Engine/Source/Developer/SessionFrontend/Private/Widgets/Browser/SSessionBrowserInstanceListRow.h
@@ -95,7 +95,6 @@ public:
 		else if (ColumnName == "Platform")
 		{
 			const PlatformInfo::FPlatformInfo* const PlatformInfo = PlatformInfo::FindPlatformInfo(*InstanceInfo->GetPlatformName());
-			check(PlatformInfo);
 
 			return SNew(SHorizontalBox)
 
@@ -103,8 +102,13 @@ public:
 				.AutoWidth()
 				.VAlign(VAlign_Center)
 				[
-					SNew(SImage)
-						.Image(FEditorStyle::GetBrush(PlatformInfo->GetIconStyleName(PlatformInfo::EPlatformIconSize::Normal)))
+					SNew(SBox)
+						.WidthOverride(24)
+						.HeightOverride(24)
+						[
+							SNew(SImage)
+								.Image((PlatformInfo) ? FEditorStyle::GetBrush(PlatformInfo->GetIconStyleName(PlatformInfo::EPlatformIconSize::Normal)) : FStyleDefaults::GetNoBrush())
+						]
 				]
 
 			+ SHorizontalBox::Slot()
diff --git a/Engine/Source/Developer/SessionLauncher/Private/Widgets/Deploy/SSessionLauncherDeviceListRow.h b/Engine/Source/Developer/SessionLauncher/Private/Widgets/Deploy/SSessionLauncherDeviceListRow.h
--- a/Engine/Source/Developer/SessionLauncher/Private/Widgets/Deploy/SSessionLauncherDeviceListRow.h
+++ b/Engine/Source/Developer/SessionLauncher/Private/Widgets/Deploy/SSessionLauncherDeviceListRow.h
@@ -75,8 +75,13 @@ public:
 			+ SHorizontalBox::Slot()
 				.AutoWidth()
 				[
-					SNew(SImage)
-						.Image(this, &SSessionLauncherDeviceListRow::HandleDeviceImage)
+					SNew(SBox)
+						.WidthOverride(24)
+						.HeightOverride(24)
+						[
+							SNew(SImage)
+								.Image(this, &SSessionLauncherDeviceListRow::HandleDeviceImage)
+						]
 				]
 
 			+ SHorizontalBox::Slot()
@@ -165,13 +170,7 @@ private:
 	const FSlateBrush* HandleDeviceImage( ) const
 	{
 		const PlatformInfo::FPlatformInfo* const PlatformInfo = PlatformInfo::FindPlatformInfo(*DeviceProxy->GetPlatformName());
-
-		if (PlatformInfo == nullptr)
-		{
-			return nullptr;
-		}
-
-		return FEditorStyle::GetBrush(PlatformInfo->GetIconStyleName(PlatformInfo::EPlatformIconSize::Normal));
+		return (PlatformInfo) ? FEditorStyle::GetBrush(PlatformInfo->GetIconStyleName(PlatformInfo::EPlatformIconSize::Normal)) : FStyleDefaults::GetNoBrush();
 	}
 
 	// Callback for getting the friendly name.
