Created a Launcher specific style for the filter box to be accissible in UFE.

diff --git a/Engine/Source/Developer/DeviceManager/Private/Widgets/Browser/SDeviceBrowserFilterBar.cpp b/Engine/Source/Developer/DeviceManager/Private/Widgets/Browser/SDeviceBrowserFilterBar.cpp
--- a/Engine/Source/Developer/DeviceManager/Private/Widgets/Browser/SDeviceBrowserFilterBar.cpp
+++ b/Engine/Source/Developer/DeviceManager/Private/Widgets/Browser/SDeviceBrowserFilterBar.cpp
@@ -37,12 +37,12 @@ void SDeviceBrowserFilterBar::Construct( const FArguments& InArgs, FDeviceBrowse
 			[
 				// platform filter
 				SNew(SComboButton)
-					.ComboButtonStyle(FEditorStyle::Get(), "ContentBrowser.Filters.Style")
+					.ComboButtonStyle(FEditorStyle::Get(), "ToolbarComboButton")
 					.ForegroundColor(FLinearColor::White)
 					.ButtonContent()
 					[
 						SNew(STextBlock)
-						.TextStyle(FEditorStyle::Get(), "ContentBrowser.Filters.Text")
+						.TextStyle(FEditorStyle::Get(), "Launcher.Filters.Text")
 						.Text(LOCTEXT("PlatformFiltersComboButtonText", "Platform Filters"))
 					]
 					.ContentPadding(0.0f)
diff --git a/Engine/Source/Editor/EditorStyle/Private/SlateEditorStyle.cpp b/Engine/Source/Editor/EditorStyle/Private/SlateEditorStyle.cpp
--- a/Engine/Source/Editor/EditorStyle/Private/SlateEditorStyle.cpp
+++ b/Engine/Source/Editor/EditorStyle/Private/SlateEditorStyle.cpp
@@ -5644,6 +5644,12 @@ void FSlateEditorStyle::FStyle::SetupAutomationStyles()
 		Set( "LauncherCommand.AdvancedBuild.Medium", new IMAGE_BRUSH("Launcher/Launcher_Advanced", Icon25x25) );
 		Set( "LauncherCommand.AdvancedBuild.Small", new IMAGE_BRUSH("Launcher/Launcher_Advanced", Icon20x20) );
 
+		Set("Launcher.Filters.Text", FTextBlockStyle(NormalText)
+			.SetFont(TTF_CORE_FONT("Fonts/Roboto-Bold", 9))
+			.SetColorAndOpacity(FLinearColor(1.0f, 1.0f, 1.0f, 0.9f))
+			.SetShadowOffset(FVector2D(1, 1))
+			.SetShadowColorAndOpacity(FLinearColor(0, 0, 0, 0.9f)));
+
 		Set( "Launcher.Platform.Warning", new IMAGE_BRUSH( "Icons/alert", Icon24x24) );
 
 #if (WITH_EDITOR || (IS_PROGRAM && PLATFORM_DESKTOP))
