Misc: UtcNow() instead of Now() should be used for time zone independent time comparisons

diff --git a/Engine/Source/Developer/DeviceManager/Private/Widgets/Processes/SDeviceProcesses.cpp b/Engine/Source/Developer/DeviceManager/Private/Widgets/Processes/SDeviceProcesses.cpp
--- a/Engine/Source/Developer/DeviceManager/Private/Widgets/Processes/SDeviceProcesses.cpp
+++ b/Engine/Source/Developer/DeviceManager/Private/Widgets/Processes/SDeviceProcesses.cpp
@@ -142,7 +142,7 @@ void SDeviceProcesses::Tick( const FGeometry& AllottedGeometry, const double InC
 {
 	SCompoundWidget::Tick(AllottedGeometry, InCurrentTime, InDeltaTime);
 
-	if (LastProcessListRefreshTime + FTimespan::FromSeconds(2.5) < FDateTime::Now())
+	if (LastProcessListRefreshTime + FTimespan::FromSeconds(2.5) < FDateTime::UtcNow())
 	{
 		ReloadProcessList(true);
 	}
@@ -227,7 +227,7 @@ void SDeviceProcesses::ReloadProcessList( bool FullyReload )
 	// refresh list view
 	ProcessTreeView->RequestTreeRefresh();
 
-	LastProcessListRefreshTime = FDateTime::Now();
+	LastProcessListRefreshTime = FDateTime::UtcNow();
 }
 
 
diff --git a/Engine/Source/Developer/SessionServices/Private/SessionManager.cpp b/Engine/Source/Developer/SessionServices/Private/SessionManager.cpp
--- a/Engine/Source/Developer/SessionServices/Private/SessionManager.cpp
+++ b/Engine/Source/Developer/SessionServices/Private/SessionManager.cpp
@@ -218,7 +218,7 @@ void FSessionManager::SendPing( )
 		MessageEndpoint->Publish(new FSessionServicePing(), EMessageScope::Network);
 	}
 
-	LastPingTime = FDateTime::Now();
+	LastPingTime = FDateTime::UtcNow();
 }
 
 
@@ -301,8 +301,9 @@ void FSessionManager::HandleSessionPongMessage( const FSessionServicePong& Messa
 
 bool FSessionManager::HandleTicker( float DeltaTime )
 {
-	FDateTime Now = FDateTime::Now();
+	FDateTime Now = FDateTime::UtcNow();
 
+	// @todo gmp: don't expire sessions for now
 //	FindExpiredSessions(Now);
 
 	if (Now >= LastPingTime + FTimespan::FromSeconds(2.5))
