UFE: Fixed non-unity build

diff --git a/Engine/Source/Developer/DeviceManager/Private/Widgets/SDeviceManager.cpp b/Engine/Source/Developer/DeviceManager/Private/Widgets/SDeviceManager.cpp
--- a/Engine/Source/Developer/DeviceManager/Private/Widgets/SDeviceManager.cpp
+++ b/Engine/Source/Developer/DeviceManager/Private/Widgets/SDeviceManager.cpp
@@ -2,6 +2,7 @@
 
 #include "DeviceManagerPrivatePCH.h"
 #include "SDockTab.h"
+#include "WorkspaceMenuStructureModule.h"
 
 
 #define LOCTEXT_NAMESPACE "SDeviceManager"
