Copying //UE4/Dev-Sequencer to //UE4/Main (Source: //UE4/Dev-Sequencer @ 3617642)

diff --git a/Engine/Plugins/Messaging/UdpMessaging/Source/UdpMessaging/Private/Transport/UdpDeserializedMessage.cpp b/Engine/Plugins/Messaging/UdpMessaging/Source/UdpMessaging/Private/Transport/UdpDeserializedMessage.cpp
--- a/Engine/Plugins/Messaging/UdpMessaging/Source/UdpMessaging/Private/Transport/UdpDeserializedMessage.cpp
+++ b/Engine/Plugins/Messaging/UdpMessaging/Source/UdpMessaging/Private/Transport/UdpDeserializedMessage.cpp
@@ -1,11 +1,15 @@
 // Copyright 1998-2017 Epic Games, Inc. All Rights Reserved.
 
 #include "Transport/UdpDeserializedMessage.h"
+
+#include "Backends/JsonStructDeserializerBackend.h"
+#include "IMessageAttachment.h"
 #include "Serialization/MemoryReader.h"
+#include "StructDeserializer.h"
+#include "Transport/UdpReassembledMessage.h"
+#include "UObject/Class.h"
 #include "UObject/Package.h"
 #include "UdpMessagingPrivate.h"
-#include "Backends/JsonStructDeserializerBackend.h"
-#include "StructDeserializer.h"
 
 
 /* FUdpDeserializedMessage structors
@@ -35,7 +39,7 @@ FUdpDeserializedMessage::~FUdpDeserializedMessage()
 /* FUdpDeserializedMessage interface
  *****************************************************************************/
 
-bool FUdpDeserializedMessage::Deserialize(const FReassembledUdpMessage& ReassembledMessage)
+bool FUdpDeserializedMessage::Deserialize(const FUdpReassembledMessage& ReassembledMessage)
 {
 	// Note that some complex values are deserialized manually here, so that we
 	// can sanity check their values. @see FUdpSerializeMessageTask::DoTask()
@@ -162,7 +166,7 @@ const TWeakObjectPtr<UScriptStruct>& FUdpDeserializedMessage::GetMessageTypeInfo
 }
 
 
-IMessageContextPtr FUdpDeserializedMessage::GetOriginalContext() const
+TSharedPtr<IMessageContext, ESPMode::ThreadSafe> FUdpDeserializedMessage::GetOriginalContext() const
 {
 	return nullptr;
 }
diff --git a/Engine/Plugins/Messaging/UdpMessaging/Source/UdpMessaging/Private/Transport/UdpDeserializedMessage.h b/Engine/Plugins/Messaging/UdpMessaging/Source/UdpMessaging/Private/Transport/UdpDeserializedMessage.h
--- a/Engine/Plugins/Messaging/UdpMessaging/Source/UdpMessaging/Private/Transport/UdpDeserializedMessage.h
+++ b/Engine/Plugins/Messaging/UdpMessaging/Source/UdpMessaging/Private/Transport/UdpDeserializedMessage.h
@@ -2,10 +2,17 @@
 
 #pragma once
 
-#include "CoreMinimal.h"
-#include "IMessageAttachment.h"
+#include "CoreTypes.h"
+#include "Containers/Array.h"
 #include "IMessageContext.h"
-#include "Transport/UdpReassembledMessage.h"
+#include "Misc/DateTime.h"
+#include "Templates/SharedPointer.h"
+#include "UObject/WeakObjectPtrTemplates.h"
+
+class FUdpReassembledMessage;
+class IMessageAttachment;
+class UScriptStruct;
+
 
 /**
  * Holds a deserialized message.
@@ -33,7 +40,7 @@ public:
 	 * @param ReassembledMessage The reassembled message to deserialize.
 	 * @return true on success, false otherwise.
 	 */
-	bool Deserialize(const FReassembledUdpMessage& ReassembledMessage);
+	bool Deserialize(const FUdpReassembledMessage& ReassembledMessage);
 
 public:
 
@@ -44,7 +51,7 @@ public:
 	virtual const FDateTime& GetExpiration() const override;
 	virtual const void* GetMessage() const override;
 	virtual const TWeakObjectPtr<UScriptStruct>& GetMessageTypeInfo() const override;
-	virtual IMessageContextPtr GetOriginalContext() const override;
+	virtual TSharedPtr<IMessageContext, ESPMode::ThreadSafe> GetOriginalContext() const override;
 	virtual const TArray<FMessageAddress>& GetRecipients() const override;
 	virtual EMessageScope GetScope() const override;
 	virtual const FMessageAddress& GetSender() const override;
diff --git a/Engine/Plugins/Messaging/UdpMessaging/Source/UdpMessaging/Private/Transport/UdpMessageBeacon.cpp b/Engine/Plugins/Messaging/UdpMessaging/Source/UdpMessaging/Private/Transport/UdpMessageBeacon.cpp
--- a/Engine/Plugins/Messaging/UdpMessaging/Source/UdpMessaging/Private/Transport/UdpMessageBeacon.cpp
+++ b/Engine/Plugins/Messaging/UdpMessaging/Source/UdpMessaging/Private/Transport/UdpMessageBeacon.cpp
@@ -1,10 +1,14 @@
 // Copyright 1998-2017 Epic Games, Inc. All Rights Reserved.
 
 #include "Transport/UdpMessageBeacon.h"
+#include "UdpMessagingPrivate.h"
+
+#include "HAL/Event.h"
 #include "HAL/RunnableThread.h"
+#include "Interfaces/IPv4/IPv4Endpoint.h"
+#include "IPAddress.h"
 #include "Serialization/ArrayWriter.h"
 #include "Sockets.h"
-#include "UdpMessagingPrivate.h"
 
 
 /* FUdpMessageHelloSender static initialization
diff --git a/Engine/Plugins/Messaging/UdpMessaging/Source/UdpMessaging/Private/Transport/UdpMessageBeacon.h b/Engine/Plugins/Messaging/UdpMessaging/Source/UdpMessaging/Private/Transport/UdpMessageBeacon.h
--- a/Engine/Plugins/Messaging/UdpMessaging/Source/UdpMessaging/Private/Transport/UdpMessageBeacon.h
+++ b/Engine/Plugins/Messaging/UdpMessaging/Source/UdpMessaging/Private/Transport/UdpMessageBeacon.h
@@ -2,15 +2,20 @@
 
 #pragma once
 
-#include "CoreMinimal.h"
+#include "CoreTypes.h"
+#include "Misc/DateTime.h"
 #include "Misc/Guid.h"
+#include "Misc/Timespan.h"
 #include "HAL/Runnable.h"
 #include "Shared/UdpMessageSegment.h"
+#include "Templates/SharedPointer.h"
 
+class FEvent;
 class FInternetAddr;
 class FSocket;
 struct FIPv4Endpoint;
 
+
 /**
  * Implements a beacon sender thread.
  *
diff --git a/Engine/Plugins/Messaging/UdpMessaging/Source/UdpMessaging/Private/Transport/UdpMessageProcessor.cpp b/Engine/Plugins/Messaging/UdpMessaging/Source/UdpMessaging/Private/Transport/UdpMessageProcessor.cpp
--- a/Engine/Plugins/Messaging/UdpMessaging/Source/UdpMessaging/Private/Transport/UdpMessageProcessor.cpp
+++ b/Engine/Plugins/Messaging/UdpMessaging/Source/UdpMessaging/Private/Transport/UdpMessageProcessor.cpp
@@ -1,12 +1,22 @@
 // Copyright 1998-2017 Epic Games, Inc. All Rights Reserved.
 
 #include "Transport/UdpMessageProcessor.h"
-#include "Serialization/ArrayWriter.h"
+#include "UdpMessagingPrivate.h"
+
 #include "Common/UdpSocketSender.h"
+#include "HAL/Event.h"
+#include "HAL/RunnableThread.h"
+#include "IMessageAttachment.h"
+#include "Serialization/ArrayReader.h"
+#include "Serialization/ArrayWriter.h"
+#include "Sockets.h"
+#include "UObject/Class.h"
+
 #include "Shared/UdpMessagingSettings.h"
 #include "Transport/UdpMessageBeacon.h"
 #include "Transport/UdpMessageSegmenter.h"
-#include "UdpMessagingPrivate.h"
+#include "Transport/UdpReassembledMessage.h"
+#include "Transport/UdpSerializedMessage.h"
 
 
 /* FUdpMessageHelloSender static initialization
@@ -76,7 +86,7 @@ FUdpMessageProcessor::~FUdpMessageProcessor()
 /* FUdpMessageProcessor interface
  *****************************************************************************/
 
-bool FUdpMessageProcessor::EnqueueInboundSegment(const FArrayReaderPtr& Data, const FIPv4Endpoint& InSender) 
+bool FUdpMessageProcessor::EnqueueInboundSegment(const TSharedPtr<FArrayReader, ESPMode::ThreadSafe>& Data, const FIPv4Endpoint& InSender) 
 {
 	if (!InboundSegments.Enqueue(FInboundSegment(Data, InSender)))
 	{
@@ -285,7 +295,7 @@ void FUdpMessageProcessor::ConsumeOutboundMessages()
 }
 
 
-bool FUdpMessageProcessor::FilterSegment(const FUdpMessageSegment::FHeader& Header, const FArrayReaderPtr& Data, const FIPv4Endpoint& InSender)
+bool FUdpMessageProcessor::FilterSegment(const FUdpMessageSegment::FHeader& Header, const TSharedPtr<FArrayReader, ESPMode::ThreadSafe>& Data, const FIPv4Endpoint& InSender)
 {
 	// filter unsupported protocol versions
 	if (Header.ProtocolVersion != UDP_MESSAGING_TRANSPORT_PROTOCOL_VERSION)
@@ -344,12 +354,12 @@ void FUdpMessageProcessor::ProcessDataSegment(FInboundSegment& Segment, FNodeInf
 		return;
 	}
 
-	TSharedPtr<FReassembledUdpMessage, ESPMode::ThreadSafe>& ReassembledMessage = NodeInfo.ReassembledMessages.FindOrAdd(DataChunk.MessageId);
+	TSharedPtr<FUdpReassembledMessage, ESPMode::ThreadSafe>& ReassembledMessage = NodeInfo.ReassembledMessages.FindOrAdd(DataChunk.MessageId);
 
 	// Reassemble message
 	if (!ReassembledMessage.IsValid())
 	{
-		ReassembledMessage = MakeShareable(new FReassembledUdpMessage(DataChunk.MessageSize, DataChunk.TotalSegments, DataChunk.Sequence, Segment.Sender));
+		ReassembledMessage = MakeShareable(new FUdpReassembledMessage(DataChunk.MessageSize, DataChunk.TotalSegments, DataChunk.Sequence, Segment.Sender));
 	}
 
 	ReassembledMessage->Reassemble(DataChunk.SegmentNumber, DataChunk.SegmentOffset, DataChunk.Data, CurrentTime);
@@ -371,7 +381,7 @@ void FUdpMessageProcessor::ProcessDataSegment(FInboundSegment& Segment, FNodeInf
 	}
 	else if (NodeInfo.Resequencer.Resequence(ReassembledMessage))
 	{
-		TSharedPtr<FReassembledUdpMessage, ESPMode::ThreadSafe> ResequencedMessage;
+		TSharedPtr<FUdpReassembledMessage, ESPMode::ThreadSafe> ResequencedMessage;
 
 		while (NodeInfo.Resequencer.Pop(ResequencedMessage))
 		{
diff --git a/Engine/Plugins/Messaging/UdpMessaging/Source/UdpMessaging/Private/Transport/UdpMessageProcessor.h b/Engine/Plugins/Messaging/UdpMessaging/Source/UdpMessaging/Private/Transport/UdpMessageProcessor.h
--- a/Engine/Plugins/Messaging/UdpMessaging/Source/UdpMessaging/Private/Transport/UdpMessageProcessor.h
+++ b/Engine/Plugins/Messaging/UdpMessaging/Source/UdpMessaging/Private/Transport/UdpMessageProcessor.h
@@ -2,21 +2,32 @@
 
 #pragma once
 
-#include "CoreMinimal.h"
-#include "Misc/Guid.h"
-#include "HAL/Runnable.h"
+#include "CoreTypes.h"
+#include "Common/UdpSocketReceiver.h"
+#include "Containers/Map.h"
 #include "Containers/Queue.h"
-#include "IMessageAttachment.h"
+#include "HAL/Runnable.h"
 #include "IMessageTransport.h"
 #include "Interfaces/IPv4/IPv4Endpoint.h"
-#include "Common/UdpSocketReceiver.h"
+#include "Misc/DateTime.h"
+#include "Misc/Guid.h"
+#include "Misc/Timespan.h"
+#include "Templates/SharedPointer.h"
+
 #include "Shared/UdpMessageSegment.h"
 #include "Transport/UdpMessageResequencer.h"
-#include "Transport/UdpSerializedMessage.h"
 
+class FArrayReader;
+class FEvent;
+class FRunnableThread;
+class FSocket;
 class FUdpMessageBeacon;
 class FUdpMessageSegmenter;
+class FUdpReassembledMessage;
+class FUdpSerializedMessage;
 class FUdpSocketSender;
+class IMessageAttachment;
+
 
 /**
  * Implements a message processor for UDP messages.
@@ -37,7 +48,7 @@ class FUdpMessageProcessor
 		FGuid NodeId;
 
 		/** Holds the collection of reassembled messages. */
-		TMap<int32, TSharedPtr<FReassembledUdpMessage, ESPMode::ThreadSafe>> ReassembledMessages;
+		TMap<int32, TSharedPtr<FUdpReassembledMessage, ESPMode::ThreadSafe>> ReassembledMessages;
 
 		/** Holds the message resequencer. */
 		FUdpMessageResequencer Resequencer;
@@ -69,7 +80,7 @@ class FUdpMessageProcessor
 	struct FInboundSegment
 	{
 		/** Holds the segment data. */
-		FArrayReaderPtr Data;
+		TSharedPtr<FArrayReader, ESPMode::ThreadSafe> Data;
 
 		/** Holds the sender's network endpoint. */
 		FIPv4Endpoint Sender;
@@ -78,7 +89,7 @@ class FUdpMessageProcessor
 		FInboundSegment() { }
 
 		/** Creates and initializes a new instance. */
-		FInboundSegment(const FArrayReaderPtr& InData, const FIPv4Endpoint& InSender)
+		FInboundSegment(const TSharedPtr<FArrayReader, ESPMode::ThreadSafe>& InData, const FIPv4Endpoint& InSender)
 			: Data(InData)
 			, Sender(InSender)
 		{ }
@@ -127,7 +138,7 @@ public:
 	 * @param Sender The sender's network endpoint.
 	 * @return true if the segment was queued up, false otherwise.
 	 */
-	bool EnqueueInboundSegment(const FArrayReaderPtr& Data, const FIPv4Endpoint& Sender);
+	bool EnqueueInboundSegment(const TSharedPtr<FArrayReader, ESPMode::ThreadSafe>& Data, const FIPv4Endpoint& Sender);
 
 	/**
 	 * Queues up an outbound message.
@@ -148,7 +159,7 @@ public:
 	 *
 	 * @return The delegate.
 	 */
-	DECLARE_DELEGATE_ThreeParams(FOnMessageReassembled, const FReassembledUdpMessage& /*ReassembledMessage*/, const IMessageAttachmentPtr& /*Attachment*/, const FGuid& /*NodeId*/)
+	DECLARE_DELEGATE_ThreeParams(FOnMessageReassembled, const FUdpReassembledMessage& /*ReassembledMessage*/, const IMessageAttachmentPtr& /*Attachment*/, const FGuid& /*NodeId*/)
 	FOnMessageReassembled& OnMessageReassembled()
 	{
 		return MessageReassembledDelegate;
@@ -160,7 +171,8 @@ public:
 	 * @return The delegate.
 	 * @see OnNodeLost
 	 */
-	IMessageTransport::FOnNodeDiscovered& OnNodeDiscovered()
+	DECLARE_DELEGATE_OneParam(FOnNodeDiscovered, const FGuid& /*NodeId*/)
+	FOnNodeDiscovered& OnNodeDiscovered()
 	{
 		return NodeDiscoveredDelegate;
 	}
@@ -171,7 +183,8 @@ public:
 	 * @return The delegate.
 	 * @see OnNodeDiscovered
 	 */
-	IMessageTransport::FOnNodeLost& OnNodeLost()
+	DECLARE_DELEGATE_OneParam(FOnNodeLost, const FGuid& /*NodeId*/)
+	FOnNodeLost& OnNodeLost()
 	{
 		return NodeLostDelegate;
 	}
@@ -217,7 +230,7 @@ protected:
 	 * @param Sender The segment sender.
 	 * @return true if the segment passed the filter, false otherwise.
 	 */
-	bool FilterSegment(const FUdpMessageSegment::FHeader& Header, const FArrayReaderPtr& Data, const FIPv4Endpoint& Sender);
+	bool FilterSegment(const FUdpMessageSegment::FHeader& Header, const TSharedPtr<FArrayReader, ESPMode::ThreadSafe>& Data, const FIPv4Endpoint& Sender);
 
 	/**
 	 * Processes an Abort segment.
@@ -361,10 +374,10 @@ private:
 	FOnMessageReassembled MessageReassembledDelegate;
 
 	/** Holds a delegate to be invoked when a network node was discovered. */
-	IMessageTransport::FOnNodeDiscovered NodeDiscoveredDelegate;
+	FOnNodeDiscovered NodeDiscoveredDelegate;
 
 	/** Holds a delegate to be invoked when a network node was lost. */
-	IMessageTransport::FOnNodeLost NodeLostDelegate;
+	FOnNodeLost NodeLostDelegate;
 
 private:
 
diff --git a/Engine/Plugins/Messaging/UdpMessaging/Source/UdpMessaging/Private/Transport/UdpMessageResequencer.h b/Engine/Plugins/Messaging/UdpMessaging/Source/UdpMessaging/Private/Transport/UdpMessageResequencer.h
--- a/Engine/Plugins/Messaging/UdpMessaging/Source/UdpMessaging/Private/Transport/UdpMessageResequencer.h
+++ b/Engine/Plugins/Messaging/UdpMessaging/Source/UdpMessaging/Private/Transport/UdpMessageResequencer.h
@@ -2,9 +2,12 @@
 
 #pragma once
 
-#include "CoreMinimal.h"
+#include "CoreTypes.h"
+#include "Containers/Array.h"
+#include "Templates/SharedPointer.h"
 #include "Transport/UdpReassembledMessage.h"
 
+
 /**
  * Implements a re-sequencer for messages received over the UDP transport.
  */
@@ -42,7 +45,7 @@ public:
 	 *
 	 * @return true if a message was returned, false otherwise.
 	 */
-	bool Pop(TSharedPtr<FReassembledUdpMessage, ESPMode::ThreadSafe>& OutMessage)
+	bool Pop(TSharedPtr<FUdpReassembledMessage, ESPMode::ThreadSafe>& OutMessage)
 	{
 		if (MessageHeap.HeapTop()->GetSequence() != NextSequence)
 		{
@@ -61,7 +64,7 @@ public:
 	 * @param Message The message to resequence.
 	 * @return true if the message is in sequence, false otherwise.
 	 */
-	bool Resequence(const TSharedPtr<FReassembledUdpMessage, ESPMode::ThreadSafe>& Message)
+	bool Resequence(const TSharedPtr<FUdpReassembledMessage, ESPMode::ThreadSafe>& Message)
 	{
 		MessageHeap.HeapPush(Message, FSequenceComparer());
 
@@ -80,7 +83,7 @@ private:
 	/** Helper for ordering messages by their sequence numbers. */
 	struct FSequenceComparer
 	{
-		bool operator()(const TSharedPtr<FReassembledUdpMessage, ESPMode::ThreadSafe>& A, const TSharedPtr<FReassembledUdpMessage, ESPMode::ThreadSafe>& B) const
+		bool operator()(const TSharedPtr<FUdpReassembledMessage, ESPMode::ThreadSafe>& A, const TSharedPtr<FUdpReassembledMessage, ESPMode::ThreadSafe>& B) const
 		{
 			return A->GetSequence() < B->GetSequence();
 		}
@@ -95,7 +98,7 @@ private:
 	uint64 HighestReceivedSequence;
 
 	/** Holds the messages that need to be resequenced. */
-	TArray<TSharedPtr<FReassembledUdpMessage, ESPMode::ThreadSafe>> MessageHeap;
+	TArray<TSharedPtr<FUdpReassembledMessage, ESPMode::ThreadSafe>> MessageHeap;
 
 	/** Holds the maximum resequence window size. */
 	uint16 WindowSize;
diff --git a/Engine/Plugins/Messaging/UdpMessaging/Source/UdpMessaging/Private/Transport/UdpMessageSegmenter.cpp b/Engine/Plugins/Messaging/UdpMessaging/Source/UdpMessaging/Private/Transport/UdpMessageSegmenter.cpp
--- a/Engine/Plugins/Messaging/UdpMessaging/Source/UdpMessaging/Private/Transport/UdpMessageSegmenter.cpp
+++ b/Engine/Plugins/Messaging/UdpMessaging/Source/UdpMessaging/Private/Transport/UdpMessageSegmenter.cpp
@@ -2,6 +2,10 @@
 
 #include "Transport/UdpMessageSegmenter.h"
 
+#include "Serialization/Archive.h"
+
+#include "Transport/UdpSerializedMessage.h"
+
 
 /* FUdpMessageSegmenter structors
  *****************************************************************************/
@@ -79,6 +83,12 @@ void FUdpMessageSegmenter::Initialize()
 }
 
 
+bool FUdpMessageSegmenter::IsInvalid() const
+{
+	return (SerializedMessage->GetState() == EUdpSerializedMessageState::Invalid);
+}
+
+
 void FUdpMessageSegmenter::MarkAsSent(uint16 Segment)
 {
 	if (Segment < PendingSegments.Num())
diff --git a/Engine/Plugins/Messaging/UdpMessaging/Source/UdpMessaging/Private/Transport/UdpMessageSegmenter.h b/Engine/Plugins/Messaging/UdpMessaging/Source/UdpMessaging/Private/Transport/UdpMessageSegmenter.h
--- a/Engine/Plugins/Messaging/UdpMessaging/Source/UdpMessaging/Private/Transport/UdpMessageSegmenter.h
+++ b/Engine/Plugins/Messaging/UdpMessaging/Source/UdpMessaging/Private/Transport/UdpMessageSegmenter.h
@@ -2,8 +2,14 @@
 
 #pragma once
 
-#include "CoreMinimal.h"
-#include "Transport/UdpSerializedMessage.h"
+#include "CoreTypes.h"
+#include "Containers/Array.h"
+#include "Containers/BitArray.h"
+#include "Templates/SharedPointer.h"
+
+class FArchive;
+class FUdpSerializedMessage;
+
 
 /**
  * Implements a message segmenter.
@@ -100,10 +106,7 @@ public:
 	 *
 	 * @return true if the segmenter is invalid, false otherwise.
 	 */
-	bool IsInvalid() const
-	{
-		return (SerializedMessage->GetState() == EUdpSerializedMessageState::Invalid);
-	}
+	bool IsInvalid() const;
 
 	/**
 	 * Marks the specified segment as sent.
diff --git a/Engine/Plugins/Messaging/UdpMessaging/Source/UdpMessaging/Private/Transport/UdpMessageTransport.cpp b/Engine/Plugins/Messaging/UdpMessaging/Source/UdpMessaging/Private/Transport/UdpMessageTransport.cpp
--- a/Engine/Plugins/Messaging/UdpMessaging/Source/UdpMessaging/Private/Transport/UdpMessageTransport.cpp
+++ b/Engine/Plugins/Messaging/UdpMessaging/Source/UdpMessaging/Private/Transport/UdpMessageTransport.cpp
@@ -1,12 +1,23 @@
 // Copyright 1998-2017 Epic Games, Inc. All Rights Reserved.
 
 #include "Transport/UdpMessageTransport.h"
+#include "UdpMessagingPrivate.h"
+
 #include "Common/UdpSocketBuilder.h"
+#include "Common/UdpSocketReceiver.h"
+#include "HAL/RunnableThread.h"
+#include "IMessageContext.h"
+#include "IMessageTransportHandler.h"
+#include "Misc/Guid.h"
+#include "Serialization/ArrayReader.h"
+#include "SocketSubsystem.h"
+#include "Sockets.h"
+
+#include "Transport/UdpReassembledMessage.h"
 #include "Transport/UdpDeserializedMessage.h"
 #include "Transport/UdpSerializedMessage.h"
 #include "Transport/UdpMessageProcessor.h"
 #include "Transport/UdpSerializeMessageTask.h"
-#include "UdpMessagingPrivate.h"
 
 
 /* FUdpMessageTransport structors
@@ -20,6 +31,7 @@ FUdpMessageTransport::FUdpMessageTransport(const FIPv4Endpoint& InUnicastEndpoin
 	, MulticastSocket(nullptr)
 	, MulticastTtl(InMulticastTtl)
 	, SocketSubsystem(ISocketSubsystem::Get(PLATFORM_SOCKETSUBSYSTEM))
+	, TransportHandler(nullptr)
 	, UnicastEndpoint(InUnicastEndpoint)
 #if PLATFORM_DESKTOP
 	, UnicastReceiver(nullptr)
@@ -37,7 +49,13 @@ FUdpMessageTransport::~FUdpMessageTransport()
 /* IMessageTransport interface
  *****************************************************************************/
 
-bool FUdpMessageTransport::StartTransport()
+FName FUdpMessageTransport::GetDebugName() const
+{
+	return "UdpMessageTransport";
+}
+
+
+bool FUdpMessageTransport::StartTransport(IMessageTransportHandler& Handler)
 {
 #if PLATFORM_DESKTOP
 	// create & initialize unicast socket (only on multi-process platforms)
@@ -82,6 +100,8 @@ bool FUdpMessageTransport::StartTransport()
 #endif
 	}
 
+	TransportHandler = &Handler;
+
 	// initialize threads
 	FTimespan ThreadWaitTime = FTimespan::FromMilliseconds(100);
 
@@ -143,6 +163,8 @@ void FUdpMessageTransport::StopTransport()
 		UnicastSocket = nullptr;
 	}
 #endif
+
+	TransportHandler = nullptr;
 }
 
 
@@ -183,19 +205,31 @@ bool FUdpMessageTransport::TransportMessage(const TSharedRef<IMessageContext, ES
 /* FUdpMessageTransport event handlers
  *****************************************************************************/
 
-void FUdpMessageTransport::HandleProcessorMessageReassembled(const FReassembledUdpMessage& ReassembledMessage, const TSharedPtr<IMessageAttachment, ESPMode::ThreadSafe>& Attachment, const FGuid& NodeId)
+void FUdpMessageTransport::HandleProcessorMessageReassembled(const FUdpReassembledMessage& ReassembledMessage, const TSharedPtr<IMessageAttachment, ESPMode::ThreadSafe>& Attachment, const FGuid& NodeId)
 {
 	// @todo gmp: move message deserialization into an async task
 	TSharedRef<FUdpDeserializedMessage, ESPMode::ThreadSafe> DeserializedMessage = MakeShareable(new FUdpDeserializedMessage(Attachment));
 
 	if (DeserializedMessage->Deserialize(ReassembledMessage))
 	{
-		MessageReceivedDelegate.ExecuteIfBound(DeserializedMessage, NodeId);
+		TransportHandler->ReceiveTransportMessage(DeserializedMessage, NodeId);
 	}
 }
 
 
-void FUdpMessageTransport::HandleSocketDataReceived(const FArrayReaderPtr& Data, const FIPv4Endpoint& Sender)
+void FUdpMessageTransport::HandleProcessorNodeDiscovered(const FGuid& DiscoveredNodeId)
+{
+	TransportHandler->DiscoverTransportNode(DiscoveredNodeId);
+}
+
+
+void FUdpMessageTransport::HandleProcessorNodeLost(const FGuid& LostNodeId)
+{
+	TransportHandler->ForgetTransportNode(LostNodeId);
+}
+
+
+void FUdpMessageTransport::HandleSocketDataReceived(const TSharedPtr<FArrayReader, ESPMode::ThreadSafe>& Data, const FIPv4Endpoint& Sender)
 {
 	if (MessageProcessor != nullptr)
 	{
diff --git a/Engine/Plugins/Messaging/UdpMessaging/Source/UdpMessaging/Private/Transport/UdpMessageTransport.h b/Engine/Plugins/Messaging/UdpMessaging/Source/UdpMessaging/Private/Transport/UdpMessageTransport.h
--- a/Engine/Plugins/Messaging/UdpMessaging/Source/UdpMessaging/Private/Transport/UdpMessageTransport.h
+++ b/Engine/Plugins/Messaging/UdpMessaging/Source/UdpMessaging/Private/Transport/UdpMessageTransport.h
@@ -2,16 +2,25 @@
 
 #pragma once
 
-#include "CoreMinimal.h"
-#include "Misc/Guid.h"
-#include "IMessageContext.h"
+#include "CoreTypes.h"
 #include "IMessageTransport.h"
 #include "Interfaces/IPv4/IPv4Endpoint.h"
 #include "IMessageAttachment.h"
-#include "Common/UdpSocketReceiver.h"
-#include "Transport/UdpReassembledMessage.h"
+#include "Templates/SharedPointer.h"
 
+class FArrayReader;
+class FUdpReassembledMessage;
+class FRunnableThread;
+class FSocket;
 class FUdpMessageProcessor;
+class FUdpSocketReceiver;
+class IMessageAttachment;
+class IMessageContext;
+class IMessageTransportHandler;
+class ISocketSubsystem;
+
+struct FGuid;
+
 
 /**
  * Implements a message transport technology using an UDP network connection.
@@ -41,49 +50,24 @@ public:
 
 	//~ IMessageTransport interface
 
-	virtual FName GetDebugName() const override
-	{
-		return "UdpMessageTransport";
-	}
-
-	virtual FOnMessageReceived& OnMessageReceived() override
-	{
-		return MessageReceivedDelegate;
-	}
-
-	virtual FOnNodeDiscovered& OnNodeDiscovered() override
-	{
-		return NodeDiscoveredDelegate;
-	}
-
-	virtual FOnNodeLost& OnNodeLost() override
-	{
-		return NodeLostDelegate;
-	}
-
-	virtual bool StartTransport() override;
+	virtual FName GetDebugName() const override;
+	virtual bool StartTransport(IMessageTransportHandler& Handler) override;
 	virtual void StopTransport() override;
 	virtual bool TransportMessage(const TSharedRef<IMessageContext, ESPMode::ThreadSafe>& Context, const TArray<FGuid>& Recipients) override;
 
 private:
 
 	/** Handles received transport messages. */
-	void HandleProcessorMessageReassembled(const FReassembledUdpMessage& ReassembledMessage, const TSharedPtr<IMessageAttachment, ESPMode::ThreadSafe>& Attachment, const FGuid& NodeId);
+	void HandleProcessorMessageReassembled(const FUdpReassembledMessage& ReassembledMessage, const TSharedPtr<IMessageAttachment, ESPMode::ThreadSafe>& Attachment, const FGuid& NodeId);
 
 	/** Handles discovered transport endpoints. */
-	void HandleProcessorNodeDiscovered(const FGuid& DiscoveredNodeId)
-	{
-		NodeDiscoveredDelegate.ExecuteIfBound(DiscoveredNodeId);
-	}
+	void HandleProcessorNodeDiscovered(const FGuid& DiscoveredNodeId);
 
 	/** Handles lost transport endpoints. */
-	void HandleProcessorNodeLost(const FGuid& LostNodeId)
-	{
-		NodeLostDelegate.ExecuteIfBound(LostNodeId);
-	}
+	void HandleProcessorNodeLost(const FGuid& LostNodeId);
 
 	/** Handles received socket data. */
-	void HandleSocketDataReceived(const FArrayReaderPtr& Data, const FIPv4Endpoint& Sender);
+	void HandleSocketDataReceived(const TSharedPtr<FArrayReader, ESPMode::ThreadSafe>& Data, const FIPv4Endpoint& Sender);
 
 private:
 
@@ -108,6 +92,9 @@ private:
 	/** Holds a pointer to the socket sub-system. */
 	ISocketSubsystem* SocketSubsystem;
 
+	/** Message transport handler. */
+	IMessageTransportHandler* TransportHandler;
+
 	/** Holds the local endpoint to receive messages on. */
 	FIPv4Endpoint UnicastEndpoint;
 
@@ -118,15 +105,4 @@ private:
 	/** Holds the unicast socket. */
 	FSocket* UnicastSocket;
 #endif
-
-private:
-
-	/** Holds a delegate to be invoked when a message was received on the transport channel. */
-	FOnMessageReceived MessageReceivedDelegate;
-
-	/** Holds a delegate to be invoked when a network node was discovered. */
-	FOnNodeDiscovered NodeDiscoveredDelegate;
-
-	/** Holds a delegate to be invoked when a network node was lost. */
-	FOnNodeLost NodeLostDelegate;
 };
diff --git a/Engine/Plugins/Messaging/UdpMessaging/Source/UdpMessaging/Private/Transport/UdpReassembledMessage.h b/Engine/Plugins/Messaging/UdpMessaging/Source/UdpMessaging/Private/Transport/UdpReassembledMessage.h
--- a/Engine/Plugins/Messaging/UdpMessaging/Source/UdpMessaging/Private/Transport/UdpReassembledMessage.h
+++ b/Engine/Plugins/Messaging/UdpMessaging/Source/UdpMessaging/Private/Transport/UdpReassembledMessage.h
@@ -2,18 +2,22 @@
 
 #pragma once
 
-#include "CoreMinimal.h"
+#include "CoreTypes.h"
+#include "Containers/Array.h"
+#include "Containers/BitArray.h"
 #include "Interfaces/IPv4/IPv4Endpoint.h"
+#include "Misc/DateTime.h"
+
 
 /**
  * Implements a reassembled message.
  */
-class FReassembledUdpMessage
+class FUdpReassembledMessage
 {
 public:
 
 	/** Default constructor. */
-	FReassembledUdpMessage() { }
+	FUdpReassembledMessage() { }
 
 	/**
 	 * Creates and initializes a new inbound message info.
@@ -23,7 +27,7 @@ public:
 	 * @param InSequence The message sequence number.
 	 * @param InSender The IPv4 endpoint of the sender.
 	 */
-	FReassembledUdpMessage(int32 MessageSize, int32 SegmentCount, uint64 InSequence, const FIPv4Endpoint& InSender)
+	FUdpReassembledMessage(int32 MessageSize, int32 SegmentCount, uint64 InSequence, const FIPv4Endpoint& InSender)
 		: PendingSegments(true, SegmentCount)
 		, PendingSegmentsCount(SegmentCount)
 		, ReceivedBytes(0)
@@ -34,7 +38,7 @@ public:
 	}
 
 	/** Virtual destructor. */
-	virtual ~FReassembledUdpMessage() { }
+	virtual ~FUdpReassembledMessage() { }
 
 public:
 
diff --git a/Engine/Plugins/Messaging/UdpMessaging/Source/UdpMessaging/Private/Transport/UdpSerializeMessageTask.cpp b/Engine/Plugins/Messaging/UdpMessaging/Source/UdpMessaging/Private/Transport/UdpSerializeMessageTask.cpp
--- a/Engine/Plugins/Messaging/UdpMessaging/Source/UdpMessaging/Private/Transport/UdpSerializeMessageTask.cpp
+++ b/Engine/Plugins/Messaging/UdpMessaging/Source/UdpMessaging/Private/Transport/UdpSerializeMessageTask.cpp
@@ -1,9 +1,13 @@
 // Copyright 1998-2017 Epic Games, Inc. All Rights Reserved.
 
 #include "Transport/UdpSerializeMessageTask.h"
+
 #include "Backends/JsonStructSerializerBackend.h"
+#include "IMessageContext.h"
 #include "StructSerializer.h"
 
+#include "Transport/UdpSerializedMessage.h"
+
 
 /* FUdpSerializeMessageTask interface
  *****************************************************************************/
diff --git a/Engine/Plugins/Messaging/UdpMessaging/Source/UdpMessaging/Private/Transport/UdpSerializeMessageTask.h b/Engine/Plugins/Messaging/UdpMessaging/Source/UdpMessaging/Private/Transport/UdpSerializeMessageTask.h
--- a/Engine/Plugins/Messaging/UdpMessaging/Source/UdpMessaging/Private/Transport/UdpSerializeMessageTask.h
+++ b/Engine/Plugins/Messaging/UdpMessaging/Source/UdpMessaging/Private/Transport/UdpSerializeMessageTask.h
@@ -2,10 +2,13 @@
 
 #pragma once
 
-#include "CoreMinimal.h"
+#include "Async/TaskGraphInterfaces.h"
 #include "Stats/Stats.h"
-#include "IMessageContext.h"
-#include "Transport/UdpSerializedMessage.h"
+#include "Templates/SharedPointer.h"
+
+class FUdpSerializedMessage;
+class IMessageContext;
+
 
 /**
  * Implements an asynchronous task for serializing a message.
diff --git a/Engine/Plugins/Messaging/UdpMessaging/Source/UdpMessaging/Private/Transport/UdpSerializedMessage.h b/Engine/Plugins/Messaging/UdpMessaging/Source/UdpMessaging/Private/Transport/UdpSerializedMessage.h
--- a/Engine/Plugins/Messaging/UdpMessaging/Source/UdpMessaging/Private/Transport/UdpSerializedMessage.h
+++ b/Engine/Plugins/Messaging/UdpMessaging/Source/UdpMessaging/Private/Transport/UdpSerializedMessage.h
@@ -2,10 +2,14 @@
 
 #pragma once
 
-#include "CoreMinimal.h"
+#include "CoreTypes.h"
+#include "Containers/Array.h"
+#include "Delegates/Delegate.h"
+#include "Serialization/Archive.h"
 #include "Serialization/MemoryWriter.h"
 #include "Serialization/MemoryReader.h"
 
+
 /**
  * Enumerates possibly states of a serialized message.
  *
diff --git a/Engine/Plugins/Messaging/UdpMessaging/Source/UdpMessaging/Private/Tunnel/UdpMessageTunnel.cpp b/Engine/Plugins/Messaging/UdpMessaging/Source/UdpMessaging/Private/Tunnel/UdpMessageTunnel.cpp
--- a/Engine/Plugins/Messaging/UdpMessaging/Source/UdpMessaging/Private/Tunnel/UdpMessageTunnel.cpp
+++ b/Engine/Plugins/Messaging/UdpMessaging/Source/UdpMessaging/Private/Tunnel/UdpMessageTunnel.cpp
@@ -1,12 +1,17 @@
 // Copyright 1998-2017 Epic Games, Inc. All Rights Reserved.
 
 #include "Tunnel/UdpMessageTunnel.h"
+#include "UdpMessagingPrivate.h"
+
 #include "Misc/ScopeLock.h"
 #include "Common/TcpSocketBuilder.h"
 #include "Common/TcpListener.h"
 #include "Common/UdpSocketBuilder.h"
+#include "Serialization/ArrayReader.h"
+
 #include "Shared/UdpMessageSegment.h"
-#include "UdpMessagingPrivate.h"
+#include "Tunnel/UdpMessageTunnelConnection.h"
+
 
 #if PLATFORM_DESKTOP
 
@@ -212,8 +217,8 @@ void FUdpMessageTunnel::RemoveExpiredNodes(TMap<FGuid, FNodeInfo>& Nodes)
 
 void FUdpMessageTunnel::TcpToUdp()
 {
-	TArray<FUdpMessageTunnelConnectionPtr>::TIterator It(Connections);
-	FArrayReaderPtr Payload;
+	TArray<TSharedPtr<FUdpMessageTunnelConnection>>::TIterator It(Connections);
+	TSharedPtr<FArrayReader, ESPMode::ThreadSafe> Payload;
 
 	while (It && UnicastSocket->Wait(ESocketWaitConditions::WaitForWrite, FTimespan::Zero()))
 	{
@@ -275,7 +280,7 @@ void FUdpMessageTunnel::UdpToTcp(FSocket* Socket)
 
 	while (Socket->HasPendingData(DatagramSize))
 	{
-		FArrayReaderPtr Datagram = MakeShareable(new FArrayReader(true));
+		TSharedRef<FArrayReader, ESPMode::ThreadSafe> Datagram = MakeShareable(new FArrayReader(true));
 		Datagram->SetNumUninitialized(FMath::Min(DatagramSize, 65507u));
 
 		int32 BytesRead = 0;
@@ -345,7 +350,7 @@ void FUdpMessageTunnel::UpdateConnections()
 	// add pending connections
 	if (!PendingConnections.IsEmpty())
 	{
-		FUdpMessageTunnelConnectionPtr PendingConnection;
+		TSharedPtr<FUdpMessageTunnelConnection> PendingConnection;
 
 		while (PendingConnections.Dequeue(PendingConnection))
 		{
diff --git a/Engine/Plugins/Messaging/UdpMessaging/Source/UdpMessaging/Private/Tunnel/UdpMessageTunnel.h b/Engine/Plugins/Messaging/UdpMessaging/Source/UdpMessaging/Private/Tunnel/UdpMessageTunnel.h
--- a/Engine/Plugins/Messaging/UdpMessaging/Source/UdpMessaging/Private/Tunnel/UdpMessageTunnel.h
+++ b/Engine/Plugins/Messaging/UdpMessaging/Source/UdpMessaging/Private/Tunnel/UdpMessageTunnel.h
@@ -2,28 +2,33 @@
 
 #pragma once
 
-#include "CoreMinimal.h"
+#include "CoreTypes.h"
 #include "Misc/Guid.h"
 #include "HAL/Runnable.h"
+#include "Containers/Map.h"
 #include "Containers/Queue.h"
 #include "Interfaces/IPv4/IPv4Endpoint.h"
-#include "Tunnel/UdpMessageTunnelConnection.h"
+#include "Templates/SharedPointer.h"
+
 #include "IUdpMessageTunnel.h"
 
+class FSocket;
 class FTcpListener;
+class FUdpMessageTunnelConnection;
+
 
 /**
  * Implements a bi-directional tunnel to send UDP messages over a TCP connection.
  */
 class FUdpMessageTunnel
-	: FRunnable
+	: public FRunnable
 	, public IUdpMessageTunnel
 {
 	/** Structure for transport node information. */
 	struct FNodeInfo
 	{
 		/** Holds the connection that owns the node (only used for remote nodes). */
-		FUdpMessageTunnelConnectionPtr Connection;
+		TSharedPtr<FUdpMessageTunnelConnection> Connection;
 
 		/** Holds the node's IP endpoint (only used for local nodes). */
 		FIPv4Endpoint Endpoint;
@@ -102,7 +107,7 @@ private:
 private:
 
 	/** Holds the list of open tunnel connections. */
-	TArray<FUdpMessageTunnelConnectionPtr> Connections;
+	TArray<TSharedPtr<FUdpMessageTunnelConnection>> Connections;
 
 	/** Holds a critical section for serializing access to the connections list. */
 	FCriticalSection CriticalSection;
@@ -123,7 +128,7 @@ private:
 	FSocket* MulticastSocket;
 
 	/** Holds a queue of pending connections. */
-	TQueue<FUdpMessageTunnelConnectionPtr, EQueueMode::Mpsc> PendingConnections;
+	TQueue<TSharedPtr<FUdpMessageTunnelConnection>, EQueueMode::Mpsc> PendingConnections;
 
 	/** Holds a collection of information for remote transport nodes. */
 	TMap<FGuid, FNodeInfo> RemoteNodes;
diff --git a/Engine/Plugins/Messaging/UdpMessaging/Source/UdpMessaging/Private/Tunnel/UdpMessageTunnelConnection.cpp b/Engine/Plugins/Messaging/UdpMessaging/Source/UdpMessaging/Private/Tunnel/UdpMessageTunnelConnection.cpp
--- a/Engine/Plugins/Messaging/UdpMessaging/Source/UdpMessaging/Private/Tunnel/UdpMessageTunnelConnection.cpp
+++ b/Engine/Plugins/Messaging/UdpMessaging/Source/UdpMessaging/Private/Tunnel/UdpMessageTunnelConnection.cpp
@@ -1,7 +1,12 @@
 // Copyright 1998-2017 Epic Games, Inc. All Rights Reserved.
 
 #include "Tunnel/UdpMessageTunnelConnection.h"
+
+#include "HAL/RunnableThread.h"
+#include "Serialization/ArrayReader.h"
 #include "Serialization/ArrayWriter.h"
+#include "Sockets.h"
+
 
 #if PLATFORM_DESKTOP
 
@@ -39,13 +44,13 @@ FUdpMessageTunnelConnection::~FUdpMessageTunnelConnection()
 /* FUdpMessageTunnelConnection interface
  *****************************************************************************/
 
-bool FUdpMessageTunnelConnection::Receive(FArrayReaderPtr& OutPayload)
+bool FUdpMessageTunnelConnection::Receive(TSharedPtr<FArrayReader, ESPMode::ThreadSafe>& OutPayload)
 {
 	return Inbox.Dequeue(OutPayload);
 }
 
 
-bool FUdpMessageTunnelConnection::Send(const FArrayReaderPtr& Payload)
+bool FUdpMessageTunnelConnection::Send(const TSharedPtr<FArrayReader, ESPMode::ThreadSafe>& Payload)
 {
 	if (IsOpen())
 	{
@@ -182,7 +187,7 @@ bool FUdpMessageTunnelConnection::ReceivePayloads()
 			TotalBytesReceived += BytesRead;
 
 			// ... receive the payload
-			FArrayReaderPtr Payload = MakeShareable(new FArrayReader(true));
+			TSharedPtr<FArrayReader, ESPMode::ThreadSafe> Payload = MakeShareable(new FArrayReader(true));
 			Payload->SetNumUninitialized(PayloadSize);
 
 			if (!Socket->Recv(Payload->GetData(), Payload->Num(), BytesRead))
@@ -216,7 +221,7 @@ bool FUdpMessageTunnelConnection::SendPayloads()
 		return true;
 	}
 
-	FArrayReaderPtr Payload;
+	TSharedPtr<FArrayReader, ESPMode::ThreadSafe> Payload;
 	Outbox.Dequeue(Payload);
 	int32 BytesSent = 0;
 
diff --git a/Engine/Plugins/Messaging/UdpMessaging/Source/UdpMessaging/Private/Tunnel/UdpMessageTunnelConnection.h b/Engine/Plugins/Messaging/UdpMessaging/Source/UdpMessaging/Private/Tunnel/UdpMessageTunnelConnection.h
--- a/Engine/Plugins/Messaging/UdpMessaging/Source/UdpMessaging/Private/Tunnel/UdpMessageTunnelConnection.h
+++ b/Engine/Plugins/Messaging/UdpMessaging/Source/UdpMessaging/Private/Tunnel/UdpMessageTunnelConnection.h
@@ -2,13 +2,20 @@
 
 #pragma once
 
-#include "CoreMinimal.h"
-#include "HAL/Runnable.h"
+#include "CoreTypes.h"
 #include "Containers/Queue.h"
+#include "HAL/Runnable.h"
 #include "Interfaces/IPv4/IPv4Endpoint.h"
-#include "Common/UdpSocketReceiver.h"
+#include "Misc/DateTime.h"
+#include "Templates/SharedPointer.h"
+
 #include "IUdpMessageTunnelConnection.h"
 
+class FArrayReader;
+class FRunnableThread;
+class FSocket;
+
+
 /**
  * Implements a UDP message tunnel connection.
  */
@@ -39,7 +46,7 @@ public:
 	 * @return true if a payload was returned, false otherwise.
 	 * @see Send
 	 */
-	bool Receive(FArrayReaderPtr& OutPayload);
+	bool Receive(TSharedPtr<FArrayReader, ESPMode::ThreadSafe>& OutPayload);
 
 	/**
 	 * Sends a payload through this connection.
@@ -47,11 +54,11 @@ public:
 	 * @param Payload The payload to send.
 	 * @see Receive
 	 */
-	bool Send(const FArrayReaderPtr& Payload);
+	bool Send(const TSharedPtr<FArrayReader, ESPMode::ThreadSafe>& Payload);
 
 public:
 
-	// FRunnable interface
+	//~ FRunnable interface
 
 	virtual bool Init() override;
 	virtual uint32 Run() override;
@@ -60,7 +67,7 @@ public:
 
 public:
 
-	// IUdpMessageTunnelConnection interface
+	//~ IUdpMessageTunnelConnection interface
 
 	virtual void Close() override;
 	virtual uint64 GetTotalBytesReceived() const override;
@@ -93,13 +100,13 @@ private:
 	FDateTime ClosedTime;
 
 	/** Holds the collection of received payloads. */
-	TQueue<FArrayReaderPtr> Inbox;
+	TQueue<TSharedPtr<FArrayReader, ESPMode::ThreadSafe>> Inbox;
 
 	/** Holds the time at which the connection was opened. */
 	FDateTime OpenedTime;
 
 	/** Holds the collection of outbound payloads. */
-	TQueue<FArrayReaderPtr> Outbox;
+	TQueue<TSharedPtr<FArrayReader, ESPMode::ThreadSafe>> Outbox;
 
 	/** Holds the IP endpoint of the remote client. */
 	FIPv4Endpoint RemoteEndpoint;
@@ -119,10 +126,3 @@ private:
 	/** thread should continue running. */
 	bool bRun;
 };
-
-
-/** Type definition for shared pointers to instances of FUdpMessageTunnelConnection. */
-typedef TSharedPtr<FUdpMessageTunnelConnection> FUdpMessageTunnelConnectionPtr;
-
-/** Type definition for shared references to instances of FUdpMessageTunnelConnection. */
-typedef TSharedRef<FUdpMessageTunnelConnection> FUdpMessageTunnelConnectionRef;
diff --git a/Engine/Plugins/Messaging/UdpMessaging/Source/UdpMessaging/Private/UdpMessagingClasses.cpp b/Engine/Plugins/Messaging/UdpMessaging/Source/UdpMessaging/Private/UdpMessagingClasses.cpp
--- a/Engine/Plugins/Messaging/UdpMessaging/Source/UdpMessaging/Private/UdpMessagingClasses.cpp
+++ b/Engine/Plugins/Messaging/UdpMessaging/Source/UdpMessaging/Private/UdpMessagingClasses.cpp
@@ -1,6 +1,5 @@
 // Copyright 1998-2017 Epic Games, Inc. All Rights Reserved.
 
-#include "CoreMinimal.h"
 #include "Shared/UdpMessagingSettings.h"
 
 
diff --git a/Engine/Plugins/Messaging/UdpMessaging/Source/UdpMessaging/Private/UdpMessagingModule.cpp b/Engine/Plugins/Messaging/UdpMessaging/Source/UdpMessaging/Private/UdpMessagingModule.cpp
--- a/Engine/Plugins/Messaging/UdpMessaging/Source/UdpMessaging/Private/UdpMessagingModule.cpp
+++ b/Engine/Plugins/Messaging/UdpMessaging/Source/UdpMessaging/Private/UdpMessagingModule.cpp
@@ -1,26 +1,26 @@
 // Copyright 1998-2017 Epic Games, Inc. All Rights Reserved.
 
-#include "CoreMinimal.h"
+#include "UdpMessagingPrivate.h"
+
+#include "CoreTypes.h"
+#include "MessageBridgeBuilder.h"
 #include "Misc/CoreMisc.h"
 #include "Misc/CommandLine.h"
 #include "Misc/CoreDelegates.h"
 #include "Misc/App.h"
 #include "Modules/ModuleInterface.h"
 #include "Modules/ModuleManager.h"
-#include "Helpers/MessageBridgeBuilder.h"
 #include "Interfaces/IPv4/IPv4Endpoint.h"
-#include "Shared/UdpMessagingSettings.h"
-#include "Transport/UdpMessageTransport.h"
-#include "Tunnel/UdpMessageTunnel.h"
-#include "UdpMessagingPrivate.h"
 
 #if WITH_EDITOR
 	#include "ISettingsModule.h"
 	#include "ISettingsSection.h"
 #endif
 
-#if PLATFORM_DESKTOP
-#endif
+#include "IUdpMessageTunnelConnection.h"
+#include "Shared/UdpMessagingSettings.h"
+#include "Transport/UdpMessageTransport.h"
+#include "Tunnel/UdpMessageTunnel.h"
 
 
 DEFINE_LOG_CATEGORY(LogUdpMessaging);
@@ -352,7 +352,10 @@ protected:
 
 		if (!FIPv4Endpoint::Parse(Settings->TunnelUnicastEndpoint, UnicastEndpoint))
 		{
-			UE_LOG(LogUdpMessaging, Warning, TEXT("Invalid setting for UnicastEndpoint '%s' - binding to all local network adapters instead"), *Settings->UnicastEndpoint);
+			if (!Settings->TunnelUnicastEndpoint.IsEmpty())
+			{
+				UE_LOG(LogUdpMessaging, Warning, TEXT("Invalid setting for UnicastEndpoint '%s' - binding to all local network adapters instead"), *Settings->UnicastEndpoint);
+			}
 
 			UnicastEndpoint = FIPv4Endpoint::Any;
 			Settings->TunnelUnicastEndpoint = UnicastEndpoint.ToString();
@@ -361,7 +364,10 @@ protected:
 
 		if (!FIPv4Endpoint::Parse(Settings->TunnelMulticastEndpoint, MulticastEndpoint))
 		{
-			UE_LOG(LogUdpMessaging, Warning, TEXT("Invalid setting for MulticastEndpoint '%s' - using default endpoint '%s' instead"), *Settings->MulticastEndpoint, *UDP_MESSAGING_DEFAULT_MULTICAST_ENDPOINT.ToText().ToString());
+			if (!Settings->TunnelMulticastEndpoint.IsEmpty())
+			{
+				UE_LOG(LogUdpMessaging, Warning, TEXT("Invalid setting for MulticastEndpoint '%s' - using default endpoint '%s' instead"), *Settings->MulticastEndpoint, *UDP_MESSAGING_DEFAULT_MULTICAST_ENDPOINT.ToText().ToString());
+			}
 
 			MulticastEndpoint = UDP_MESSAGING_DEFAULT_MULTICAST_ENDPOINT;
 			Settings->TunnelMulticastEndpoint = MulticastEndpoint.ToString();
diff --git a/Engine/Plugins/Messaging/UdpMessaging/Source/UdpMessaging/Private/UdpMessagingPrivate.h b/Engine/Plugins/Messaging/UdpMessaging/Source/UdpMessaging/Private/UdpMessagingPrivate.h
--- a/Engine/Plugins/Messaging/UdpMessaging/Source/UdpMessaging/Private/UdpMessagingPrivate.h
+++ b/Engine/Plugins/Messaging/UdpMessaging/Source/UdpMessaging/Private/UdpMessagingPrivate.h
@@ -2,18 +2,11 @@
 
 #pragma once
 
-
-/* Private dependencies
- *****************************************************************************/
-
-#include "CoreMinimal.h"
 #include "Interfaces/IPv4/IPv4Address.h"
 #include "Interfaces/IPv4/IPv4Endpoint.h"
+#include "Logging/LogMacros.h"
 
 
-/* Private constants
- *****************************************************************************/
-
 /** Declares a log category for this module. */
 DECLARE_LOG_CATEGORY_EXTERN(LogUdpMessaging, Log, All);
 
@@ -32,17 +25,3 @@ DECLARE_LOG_CATEGORY_EXTERN(LogUdpMessaging, Log, All);
 
 /** Defines the protocol version of the UDP message transport. */
 #define UDP_MESSAGING_TRANSPORT_PROTOCOL_VERSION 10
-
-
-/* Private includes
- *****************************************************************************/
-
-// shared
-
-// transport
-
-#if PLATFORM_DESKTOP
-
-// tunnel
-
-#endif
diff --git a/Engine/Plugins/Messaging/UdpMessaging/Source/UdpMessaging/Public/IUdpMessageTunnel.h b/Engine/Plugins/Messaging/UdpMessaging/Source/UdpMessaging/Public/IUdpMessageTunnel.h
--- a/Engine/Plugins/Messaging/UdpMessaging/Source/UdpMessaging/Public/IUdpMessageTunnel.h
+++ b/Engine/Plugins/Messaging/UdpMessaging/Source/UdpMessaging/Public/IUdpMessageTunnel.h
@@ -2,11 +2,16 @@
 
 #pragma once
 
-#include "CoreMinimal.h"
+#include "CoreTypes.h"
+#include "Containers/Array.h"
+#include "Delegate.h"
+#include "Templates/SharedPointer.h"
 
 class IUdpMessageTunnelConnection;
+
 struct FIPv4Endpoint;
 
+
 /**
  * Interface for message tunnels.
  */
diff --git a/Engine/Plugins/Messaging/UdpMessaging/Source/UdpMessaging/Public/IUdpMessageTunnelConnection.h b/Engine/Plugins/Messaging/UdpMessaging/Source/UdpMessaging/Public/IUdpMessageTunnelConnection.h
--- a/Engine/Plugins/Messaging/UdpMessaging/Source/UdpMessaging/Public/IUdpMessageTunnelConnection.h
+++ b/Engine/Plugins/Messaging/UdpMessaging/Source/UdpMessaging/Public/IUdpMessageTunnelConnection.h
@@ -2,7 +2,10 @@
 
 #pragma once
 
-#include "CoreMinimal.h"
+#include "CoreTypes.h"
+#include "Internationalization/Text.h"
+#include "Misc/Timespan.h"
+
 
 /**
  * Interface for message tunnel connections.
diff --git a/Engine/Plugins/Messaging/UdpMessaging/Source/UdpMessaging/UdpMessaging.Build.cs b/Engine/Plugins/Messaging/UdpMessaging/Source/UdpMessaging/UdpMessaging.Build.cs
--- a/Engine/Plugins/Messaging/UdpMessaging/Source/UdpMessaging/UdpMessaging.Build.cs
+++ b/Engine/Plugins/Messaging/UdpMessaging/Source/UdpMessaging/UdpMessaging.Build.cs
@@ -6,18 +6,16 @@ namespace UnrealBuildTool.Rules
 	{
 		public UdpMessaging(ReadOnlyTargetRules Target) : base(Target)
 		{
-            DynamicallyLoadedModuleNames.AddRange(
-                new string[] {
-                    "Messaging",
-				}
-            );
+			DynamicallyLoadedModuleNames.AddRange(
+				new string[] {
+					"Messaging",
+				});
 
 			PublicDependencyModuleNames.AddRange(
 				new string[]
 				{
 					"Core",
-				}
-			); 
+				});
 
 			PrivateDependencyModuleNames.AddRange(
 				new string[] {
@@ -26,14 +24,13 @@ namespace UnrealBuildTool.Rules
 					"Networking",
 					"Serialization",
 					"Sockets",
-				}
-			);
+				});
 
 			PrivateIncludePathModuleNames.AddRange(
 				new string[] {
 					"Messaging",
-				}
-			);
+					"MessagingCommon",
+				});
 
 			PrivateIncludePaths.AddRange(
 				new string[] {
@@ -42,22 +39,19 @@ namespace UnrealBuildTool.Rules
 					"UdpMessaging/Private/Transport",
 					"UdpMessaging/Private/Transport/Tests",
 					"UdpMessaging/Private/Tunnel",
-				}
-			);
+				});
 
 			if (Target.Type == TargetType.Editor)
 			{
 				DynamicallyLoadedModuleNames.AddRange(
 					new string[] {
-					"Settings",
-				}
-				);
+						"Settings",
+					});
 
 				PrivateIncludePathModuleNames.AddRange(
 					new string[] {
-					"Settings",
-				}
-				);
+						"Settings",
+					});
 			}
 		}
 	}
diff --git a/Engine/Plugins/Online/Android/OnlineSubsystemGooglePlay/Source/Private/OnlineAsyncTaskGooglePlayShowLoginUI.h b/Engine/Plugins/Online/Android/OnlineSubsystemGooglePlay/Source/Private/OnlineAsyncTaskGooglePlayShowLoginUI.h
--- a/Engine/Plugins/Online/Android/OnlineSubsystemGooglePlay/Source/Private/OnlineAsyncTaskGooglePlayShowLoginUI.h
+++ b/Engine/Plugins/Online/Android/OnlineSubsystemGooglePlay/Source/Private/OnlineAsyncTaskGooglePlayShowLoginUI.h
@@ -10,6 +10,9 @@
 #include "AndroidPermissionFunctionLibrary.h"
 
 THIRD_PARTY_INCLUDES_START
+// begin missing GPG dependencies
+#include <vector>
+// end missing GPG depdencies
 #include "gpg/types.h"
 #include "gpg/player_manager.h"
 THIRD_PARTY_INCLUDES_END
diff --git a/Engine/Plugins/Online/OnlineSubsystemFacebook/Source/Private/IOS/OnlineIdentityFacebook.cpp b/Engine/Plugins/Online/OnlineSubsystemFacebook/Source/Private/IOS/OnlineIdentityFacebook.cpp
--- a/Engine/Plugins/Online/OnlineSubsystemFacebook/Source/Private/IOS/OnlineIdentityFacebook.cpp
+++ b/Engine/Plugins/Online/OnlineSubsystemFacebook/Source/Private/IOS/OnlineIdentityFacebook.cpp
@@ -11,6 +11,7 @@
 
 // Other UE4 includes
 #include "IOSAppDelegate.h"
+#include "Misc/ConfigCacheIni.h"
 #include "IOS/IOSAsyncTask.h"
 #include "Misc/ConfigCacheIni.h"
 
diff --git a/Engine/Plugins/Online/OnlineSubsystemFacebook/Source/Private/IOS/OnlineSubsystemFacebook.cpp b/Engine/Plugins/Online/OnlineSubsystemFacebook/Source/Private/IOS/OnlineSubsystemFacebook.cpp
--- a/Engine/Plugins/Online/OnlineSubsystemFacebook/Source/Private/IOS/OnlineSubsystemFacebook.cpp
+++ b/Engine/Plugins/Online/OnlineSubsystemFacebook/Source/Private/IOS/OnlineSubsystemFacebook.cpp
@@ -2,6 +2,9 @@
 
 #include "OnlineSubsystemFacebook.h"
 #include "OnlineSubsystemFacebookPrivate.h"
+
+#include "OnlineFriendsFacebook.h"
+#include "OnlineIdentityFacebook.h"
 #include "OnlineSharingFacebook.h"
 #include "OnlineUserFacebook.h"
 
diff --git a/Engine/Plugins/Online/OnlineSubsystemOculus/Source/Private/OnlineMessageMultiTaskOculus.cpp b/Engine/Plugins/Online/OnlineSubsystemOculus/Source/Private/OnlineMessageMultiTaskOculus.cpp
--- a/Engine/Plugins/Online/OnlineSubsystemOculus/Source/Private/OnlineMessageMultiTaskOculus.cpp
+++ b/Engine/Plugins/Online/OnlineSubsystemOculus/Source/Private/OnlineMessageMultiTaskOculus.cpp
@@ -1,10 +1,9 @@
 // Copyright 1998-2017 Epic Games, Inc. All Rights Reserved.
 
-#pragma once
-
 #include "OnlineMessageMultiTaskOculus.h"
 #include "OnlineSubsystemOculusPrivate.h"
 
+
 void FOnlineMessageMultiTaskOculus::AddNewRequest(ovrRequest RequestId)
 {
 	InProgressRequests.Add(RequestId);
diff --git a/Engine/Plugins/Runtime/AndroidMoviePlayer/Source/AndroidMoviePlayer/Private/AndroidMovieStreamer.cpp b/Engine/Plugins/Runtime/AndroidMoviePlayer/Source/AndroidMoviePlayer/Private/AndroidMovieStreamer.cpp
--- a/Engine/Plugins/Runtime/AndroidMoviePlayer/Source/AndroidMoviePlayer/Private/AndroidMovieStreamer.cpp
+++ b/Engine/Plugins/Runtime/AndroidMoviePlayer/Source/AndroidMoviePlayer/Private/AndroidMovieStreamer.cpp
@@ -23,7 +23,7 @@ FAndroidMediaPlayerStreamer::FAndroidMediaPlayerStreamer()
 	: JavaMediaPlayer(nullptr)
 	, CurrentPosition(-1)
 {
-	JavaMediaPlayer = MakeShareable(new FJavaAndroidMediaPlayer(FAndroidMisc::ShouldUseVulkan()));
+	JavaMediaPlayer = MakeShareable(new FJavaAndroidMediaPlayer(false, FAndroidMisc::ShouldUseVulkan()));
 	MovieViewport = MakeShareable(new FMovieViewport());
 }
 
diff --git a/Engine/Plugins/Runtime/HTTPChunkInstaller/Source/Public/ChunkInstall.h b/Engine/Plugins/Runtime/HTTPChunkInstaller/Source/Public/ChunkInstall.h
--- a/Engine/Plugins/Runtime/HTTPChunkInstaller/Source/Public/ChunkInstall.h
+++ b/Engine/Plugins/Runtime/HTTPChunkInstaller/Source/Public/ChunkInstall.h
@@ -146,7 +146,7 @@ public:
 
 	bool IsDone()
 	{
-		return CompleteEvent->Wait(FTimespan(0));
+		return CompleteEvent->Wait(FTimespan::Zero());
 	}
 
 	static const TCHAR *Name()
diff --git a/Engine/Plugins/Runtime/HTTPChunkInstaller/Source/Public/ChunkSetup.h b/Engine/Plugins/Runtime/HTTPChunkInstaller/Source/Public/ChunkSetup.h
--- a/Engine/Plugins/Runtime/HTTPChunkInstaller/Source/Public/ChunkSetup.h
+++ b/Engine/Plugins/Runtime/HTTPChunkInstaller/Source/Public/ChunkSetup.h
@@ -181,7 +181,7 @@ public:
 
 	bool IsDone()
 	{
-		return CompleteEvent->Wait(FTimespan(0));
+		return CompleteEvent->Wait(FTimespan::Zero());
 	}
 };
 
@@ -334,6 +334,6 @@ public:
 
 	bool IsDone()
 	{
-		return CompleteEvent->Wait(FTimespan(0));
+		return CompleteEvent->Wait(FTimespan::Zero());
 	}
 };
\ No newline at end of file
diff --git a/Engine/Plugins/Runtime/MixedRealityFramework/Source/MixedRealityFramework/MixedRealityFramework.Build.cs b/Engine/Plugins/Runtime/MixedRealityFramework/Source/MixedRealityFramework/MixedRealityFramework.Build.cs
--- a/Engine/Plugins/Runtime/MixedRealityFramework/Source/MixedRealityFramework/MixedRealityFramework.Build.cs
+++ b/Engine/Plugins/Runtime/MixedRealityFramework/Source/MixedRealityFramework/MixedRealityFramework.Build.cs
@@ -28,10 +28,8 @@ public class MixedRealityFramework : ModuleRules
 				"Core",
 				"CoreUObject",
 				"Engine",
-
 				"Media",
 				"MediaAssets",
-				"MediaUtils",
 				"HeadMountedDisplay",
 				"InputCore"
 			}
diff --git a/Engine/Plugins/Runtime/MixedRealityFramework/Source/MixedRealityFramework/Private/MixedRealityCaptureComponent.cpp b/Engine/Plugins/Runtime/MixedRealityFramework/Source/MixedRealityFramework/Private/MixedRealityCaptureComponent.cpp
--- a/Engine/Plugins/Runtime/MixedRealityFramework/Source/MixedRealityFramework/Private/MixedRealityCaptureComponent.cpp
+++ b/Engine/Plugins/Runtime/MixedRealityFramework/Source/MixedRealityFramework/Private/MixedRealityCaptureComponent.cpp
@@ -5,8 +5,6 @@
 #include "MixedRealityBillboard.h"
 #include "Engine/TextureRenderTarget2D.h"
 #include "MediaPlayer.h"
-#include "IMediaPlayer.h"
-#include "IMediaTracks.h"
 #include "Materials/Material.h"
 #include "UObject/ConstructorHelpers.h"
 #include "GameFramework/Pawn.h" // for GetWorld()
@@ -290,14 +288,11 @@ bool UMixedRealityCaptureComponent::GetEditorPreviewInfo(float /*DeltaTime*/, FM
 float UMixedRealityCaptureComponent::GetDesiredAspectRatio() const
 {
 	float DesiredAspectRatio = 0.0f;
+
 	if (MediaSource)
 	{
-		TSharedPtr<IMediaPlayer, ESPMode::ThreadSafe> MediaPlayer = MediaSource->GetBasePlayer().GetNativePlayer();
-		if (MediaPlayer.IsValid())
-		{
-			IMediaTracks& Tracks = MediaPlayer->GetTracks();
-			DesiredAspectRatio = Tracks.GetVideoTrackAspectRatio(0);
-		}
+		const int32 SelectedTrack = MediaSource->GetSelectedTrack(EMediaPlayerTrack::Video);
+		DesiredAspectRatio = MediaSource->GetVideoTrackAspectRatio(SelectedTrack, MediaSource->GetTrackFormat(EMediaPlayerTrack::Video, SelectedTrack));
 	}
 
 	if (DesiredAspectRatio == 0.0f)
@@ -311,6 +306,7 @@ float UMixedRealityCaptureComponent::GetDesiredAspectRatio() const
 			DesiredAspectRatio = 16.f / 9.f;
 		}
 	}
+
 	return DesiredAspectRatio;
 }
 
diff --git a/Engine/Plugins/Runtime/Oculus/OculusVR/Source/OculusHMD/Private/OculusSceneCaptureCubemap.cpp b/Engine/Plugins/Runtime/Oculus/OculusVR/Source/OculusHMD/Private/OculusSceneCaptureCubemap.cpp
--- a/Engine/Plugins/Runtime/Oculus/OculusVR/Source/OculusHMD/Private/OculusSceneCaptureCubemap.cpp
+++ b/Engine/Plugins/Runtime/Oculus/OculusVR/Source/OculusHMD/Private/OculusSceneCaptureCubemap.cpp
@@ -2,6 +2,7 @@
 
 #include "OculusSceneCaptureCubemap.h"
 #include "OculusHMDPrivate.h"
+#include "IImageWrapper.h"
 #include "IImageWrapperModule.h"
 #include "Kismet/GameplayStatics.h"
 #include "GameFramework/PlayerController.h"
@@ -112,7 +113,7 @@ void UOculusSceneCaptureCubemap::Tick(float DeltaTime)
 
 	//Read Whole Capture Buffer
 	IImageWrapperModule& ImageWrapperModule = FModuleManager::LoadModuleChecked<IImageWrapperModule>(FName("ImageWrapper"));
-	IImageWrapperPtr ImageWrapper = ImageWrapperModule.CreateImageWrapper(EImageFormat::PNG);
+	TSharedPtr<IImageWrapper> ImageWrapper = ImageWrapperModule.CreateImageWrapper(EImageFormat::PNG);
 
 	TArray<FColor> OneFaceSurface, WholeCubemapData;
 	OneFaceSurface.AddUninitialized(CaptureBoxSideRes * CaptureBoxSideRes);
diff --git a/Engine/Plugins/Tests/EditorTests/Source/EditorTests/Private/Animation/AnimationBlueprintFastPathTest.cpp b/Engine/Plugins/Tests/EditorTests/Source/EditorTests/Private/Animation/AnimationBlueprintFastPathTest.cpp
--- a/Engine/Plugins/Tests/EditorTests/Source/EditorTests/Private/Animation/AnimationBlueprintFastPathTest.cpp
+++ b/Engine/Plugins/Tests/EditorTests/Source/EditorTests/Private/Animation/AnimationBlueprintFastPathTest.cpp
@@ -1,6 +1,6 @@
 // Copyright 1998-2017 Epic Games, Inc. All Rights Reserved.
 
-#include "AutomationTest.h"
+#include "Misc/AutomationTest.h"
 #include "KismetEditorUtilities.h"
 #include "Animation/AnimBlueprint.h"
 #include "Animation/AnimClassInterface.h"
diff --git a/Engine/Plugins/Tests/RuntimeTests/Source/RuntimeTests/Private/AutomationSpec.spec.cpp b/Engine/Plugins/Tests/RuntimeTests/Source/RuntimeTests/Private/AutomationSpec.spec.cpp
--- a/Engine/Plugins/Tests/RuntimeTests/Source/RuntimeTests/Private/AutomationSpec.spec.cpp
+++ b/Engine/Plugins/Tests/RuntimeTests/Source/RuntimeTests/Private/AutomationSpec.spec.cpp
@@ -1,6 +1,6 @@
 // Copyright 1998-2017 Epic Games, Inc. All Rights Reserved.
 
-#include "AutomationTest.h"
+#include "Misc/AutomationTest.h"
 
 BEGIN_DEFINE_SPEC(AutomationSpec, "System.Automation.Spec", EAutomationTestFlags::EngineFilter | EAutomationTestFlags::ApplicationContextMask)
 	bool Foo;
diff --git a/Engine/Shaders/Private/MaterialTemplate.ush b/Engine/Shaders/Private/MaterialTemplate.ush
--- a/Engine/Shaders/Private/MaterialTemplate.ush
+++ b/Engine/Shaders/Private/MaterialTemplate.ush
@@ -896,6 +896,11 @@ float GetSphericalParticleOpacity(FMaterialPixelParameters Parameters, float Den
 	return Opacity;
 }
 
+float2 RotateScaleOffsetTexCoords(float2 InTexCoords, float4 InRotationScale, float2 InOffset)
+{
+	return float2(dot(InTexCoords, InRotationScale.xy), dot(InTexCoords, InRotationScale.zw)) + InOffset;
+}
+
 #if USES_SPEEDTREE
 
 /** Vertex offset for SpeedTree wind and LOD */
diff --git a/Engine/Shaders/Private/MediaShaders.usf b/Engine/Shaders/Private/MediaShaders.usf
--- a/Engine/Shaders/Private/MediaShaders.usf
+++ b/Engine/Shaders/Private/MediaShaders.usf
@@ -22,12 +22,7 @@ float3 YuvToRgb(float3 YUV, float4x4 ColorTransform, uint SrgbToLinear)
 	YUV -= float3(0.0625, 0.5, 0.5);
 	float3 sRGB = mul(Coeff, YUV);
 
-	if (SrgbToLinear == 0)
-	{
-		return sRGB;
-	}
-
-	return sRGBToLinear(sRGB);
+	return (SrgbToLinear == 1) ? sRGBToLinear(sRGB) : sRGB;
 }
 
 
@@ -69,7 +64,10 @@ void BMPConvertPS(
 {
 	InUV.y = 1.0f - InUV.y;
 	float2 ScaledUV = InUV * BMPConvertUB.UVScale;
-	OutColor = BMPConvertUB.Texture.Sample(BMPConvertUB.Sampler, ScaledUV);
+	float4 RGB = BMPConvertUB.Texture.Sample(BMPConvertUB.Sampler, ScaledUV);
+
+	OutColor.xyz = (BMPConvertUB.SrgbToLinear == 1) ? sRGBToLinear(RGB.xyz) : RGB.xyz;
+	OutColor.w = RGB.w;
 }
 
 
@@ -132,9 +130,7 @@ void RGBConvertPS(
 	out float4 OutColor : SV_Target0)
 {
 	float2 ScaledUV = InUV * RGBConvertUB.UVScale;
-	float4 RGB = RGBConvertUB.Texture.Sample(RGBConvertUB.Sampler, ScaledUV);
-
-	OutColor.xyz = (RGBConvertUB.SrgbToLinear == 1) ? sRGBToLinear(RGB) : RGB;
+	OutColor = RGBConvertUB.Texture.Sample(RGBConvertUB.Sampler, ScaledUV);
 }
 
 
diff --git a/Engine/Source/Developer/AutomationController/AutomationController.Build.cs b/Engine/Source/Developer/AutomationController/AutomationController.Build.cs
--- a/Engine/Source/Developer/AutomationController/AutomationController.Build.cs
+++ b/Engine/Source/Developer/AutomationController/AutomationController.Build.cs
@@ -42,7 +42,7 @@ namespace UnrealBuildTool.Rules
             PrivateIncludePathModuleNames.AddRange(
 				new string[]
 				{
-					"Messaging",
+					"MessagingCommon",
 				}
 			);
 
diff --git a/Engine/Source/Developer/AutomationController/Private/AutomationCommandline.cpp b/Engine/Source/Developer/AutomationController/Private/AutomationCommandline.cpp
--- a/Engine/Source/Developer/AutomationController/Private/AutomationCommandline.cpp
+++ b/Engine/Source/Developer/AutomationController/Private/AutomationCommandline.cpp
@@ -9,7 +9,7 @@
 #include "Misc/App.h"
 #include "Modules/ModuleManager.h"
 #include "Misc/FilterCollection.h"
-#include "Interfaces/IAutomationControllerModule.h"
+#include "IAutomationControllerModule.h"
 #include "FileManager.h"
 #include "Paths.h"
 #include "FileHelper.h"
diff --git a/Engine/Source/Developer/AutomationController/Private/AutomationControllerManager.h b/Engine/Source/Developer/AutomationController/Private/AutomationControllerManager.h
--- a/Engine/Source/Developer/AutomationController/Private/AutomationControllerManager.h
+++ b/Engine/Source/Developer/AutomationController/Private/AutomationControllerManager.h
@@ -7,9 +7,9 @@
 #include "Misc/Guid.h"
 #include "Containers/Queue.h"
 #include "Misc/AutomationTest.h"
-#include "Interfaces/IAutomationControllerManager.h"
+#include "IAutomationControllerManager.h"
 #include "IMessageContext.h"
-#include "Helpers/MessageEndpoint.h"
+#include "MessageEndpoint.h"
 #include "Developer/AutomationController/Private/AutomationDeviceClusterManager.h"
 #include "Developer/AutomationController/Private/AutomationReportManager.h"
 #include "Async/Future.h"
@@ -391,31 +391,31 @@ protected:
 private:
 
 	/** Handles FAutomationWorkerFindWorkersResponse messages. */
-	void HandleFindWorkersResponseMessage( const FAutomationWorkerFindWorkersResponse& Message, const IMessageContextRef& Context );
+	void HandleFindWorkersResponseMessage( const FAutomationWorkerFindWorkersResponse& Message, const TSharedRef<IMessageContext, ESPMode::ThreadSafe>& Context );
 
 	/** Handles FAutomationWorkerPong messages. */
-	void HandlePongMessage( const FAutomationWorkerPong& Message, const IMessageContextRef& Context );
+	void HandlePongMessage( const FAutomationWorkerPong& Message, const TSharedRef<IMessageContext, ESPMode::ThreadSafe>& Context );
 
 	/** Handles FAutomationWorkerScreenImage messages. */
-	void HandleReceivedScreenShot( const FAutomationWorkerScreenImage& Message, const IMessageContextRef& Context );
+	void HandleReceivedScreenShot( const FAutomationWorkerScreenImage& Message, const TSharedRef<IMessageContext, ESPMode::ThreadSafe>& Context );
 
 	/** Handles FAutomationWorkerTestDataRequest messages. */
-	void HandleTestDataRequest(const FAutomationWorkerTestDataRequest& Message, const IMessageContextRef& Context);
+	void HandleTestDataRequest(const FAutomationWorkerTestDataRequest& Message, const TSharedRef<IMessageContext, ESPMode::ThreadSafe>& Context);
 
 	/** Handles FAutomationWorkerPerformanceDataRequest messages. */
-	void HandlePerformanceDataRequest(const FAutomationWorkerPerformanceDataRequest& Message, const IMessageContextRef& Context);
+	void HandlePerformanceDataRequest(const FAutomationWorkerPerformanceDataRequest& Message, const TSharedRef<IMessageContext, ESPMode::ThreadSafe>& Context);
 
 	/** Handles FAutomationWorkerRequestNextNetworkCommand messages. */
-	void HandleRequestNextNetworkCommandMessage( const FAutomationWorkerRequestNextNetworkCommand& Message, const IMessageContextRef& Context );
+	void HandleRequestNextNetworkCommandMessage( const FAutomationWorkerRequestNextNetworkCommand& Message, const TSharedRef<IMessageContext, ESPMode::ThreadSafe>& Context );
 
 	/** Handles FAutomationWorkerRequestTestsReplyComplete messages. */
-	void HandleRequestTestsReplyCompleteMessage(const FAutomationWorkerRequestTestsReplyComplete& Message, const IMessageContextRef& Context);
+	void HandleRequestTestsReplyCompleteMessage(const FAutomationWorkerRequestTestsReplyComplete& Message, const TSharedRef<IMessageContext, ESPMode::ThreadSafe>& Context);
 
 	/** Handles FAutomationWorkerRunTestsReply messages. */
-	void HandleRunTestsReplyMessage( const FAutomationWorkerRunTestsReply& Message, const IMessageContextRef& Context );
+	void HandleRunTestsReplyMessage( const FAutomationWorkerRunTestsReply& Message, const TSharedRef<IMessageContext, ESPMode::ThreadSafe>& Context );
 
 	/** Handles FAutomationWorkerWorkerOffline messages. */
-	void HandleWorkerOfflineMessage( const FAutomationWorkerWorkerOffline& Message, const IMessageContextRef& Context );
+	void HandleWorkerOfflineMessage( const FAutomationWorkerWorkerOffline& Message, const TSharedRef<IMessageContext, ESPMode::ThreadSafe>& Context );
 
 private:
 
@@ -465,7 +465,7 @@ private:
 	double LastTimeUpdateTicked;
 
 	/** Holds the messaging endpoint. */
-	FMessageEndpointPtr MessageEndpoint;
+	TSharedPtr<FMessageEndpoint, ESPMode::ThreadSafe> MessageEndpoint;
 
 	/** Counter for number of workers we have received responses from for Refreshing the Test List */
 	uint32 RefreshTestResponses;
diff --git a/Engine/Source/Developer/AutomationController/Private/AutomationControllerManger.cpp b/Engine/Source/Developer/AutomationController/Private/AutomationControllerManger.cpp
--- a/Engine/Source/Developer/AutomationController/Private/AutomationControllerManger.cpp
+++ b/Engine/Source/Developer/AutomationController/Private/AutomationControllerManger.cpp
@@ -7,12 +7,12 @@
 #include "Misc/ConfigCacheIni.h"
 #include "Misc/AutomationTest.h"
 #include "Misc/App.h"
-#include "Interfaces/IAutomationReport.h"
+#include "IAutomationReport.h"
 #include "AutomationWorkerMessages.h"
 #include "IMessageContext.h"
-#include "Helpers/MessageEndpoint.h"
+#include "MessageEndpoint.h"
+#include "MessageEndpointBuilder.h"
 #include "Modules/ModuleManager.h"
-#include "Helpers/MessageEndpointBuilder.h"
 #include "AssetEditorMessages.h"
 #include "ImageComparer.h"
 #include "AutomationControllerManager.h"
@@ -855,7 +855,7 @@ bool FAutomationControllerManager::IsTestRunnable(IAutomationReportPtr InReport)
 /* FAutomationControllerModule callbacks
  *****************************************************************************/
 
-void FAutomationControllerManager::HandleFindWorkersResponseMessage(const FAutomationWorkerFindWorkersResponse& Message, const IMessageContextRef& Context)
+void FAutomationControllerManager::HandleFindWorkersResponseMessage(const FAutomationWorkerFindWorkersResponse& Message, const TSharedRef<IMessageContext, ESPMode::ThreadSafe>& Context)
 {
 	if ( Message.SessionId == ActiveSessionId )
 	{
@@ -867,12 +867,12 @@ void FAutomationControllerManager::HandleFindWorkersResponseMessage(const FAutom
 	SetControllerStatus(EAutomationControllerModuleState::Ready);
 }
 
-void FAutomationControllerManager::HandlePongMessage( const FAutomationWorkerPong& Message, const IMessageContextRef& Context )
+void FAutomationControllerManager::HandlePongMessage( const FAutomationWorkerPong& Message, const TSharedRef<IMessageContext, ESPMode::ThreadSafe>& Context )
 {
 	AddPingResult(Context->GetSender());
 }
 
-void FAutomationControllerManager::HandleReceivedScreenShot(const FAutomationWorkerScreenImage& Message, const IMessageContextRef& Context)
+void FAutomationControllerManager::HandleReceivedScreenShot(const FAutomationWorkerScreenImage& Message, const TSharedRef<IMessageContext, ESPMode::ThreadSafe>& Context)
 {
 	FString ScreenshotIncomingFolder = FPaths::ProjectSavedDir() / TEXT("Automation/Incoming/");
 
@@ -899,7 +899,7 @@ void FAutomationControllerManager::HandleReceivedScreenShot(const FAutomationWor
 	ComparisonQueue.Enqueue(Comparison);
 }
 
-void FAutomationControllerManager::HandleTestDataRequest(const FAutomationWorkerTestDataRequest& Message, const IMessageContextRef& Context)
+void FAutomationControllerManager::HandleTestDataRequest(const FAutomationWorkerTestDataRequest& Message, const TSharedRef<IMessageContext, ESPMode::ThreadSafe>& Context)
 {
 	const FString TestDataRoot = FPaths::ConvertRelativePathToFull(FPaths::ProjectDir() / TEXT("Test"));
 	const FString DataFile = Message.DataType / Message.DataPlatform / Message.DataTestName / Message.DataName + TEXT(".json");
@@ -944,7 +944,7 @@ void FAutomationControllerManager::HandleTestDataRequest(const FAutomationWorker
 	MessageEndpoint->Send(ResponseMessage, Context->GetSender());
 }
 
-void FAutomationControllerManager::HandlePerformanceDataRequest(const FAutomationWorkerPerformanceDataRequest& Message, const IMessageContextRef& Context)
+void FAutomationControllerManager::HandlePerformanceDataRequest(const FAutomationWorkerPerformanceDataRequest& Message, const TSharedRef<IMessageContext, ESPMode::ThreadSafe>& Context)
 {
 	//TODO Read/Performance data.
 
@@ -955,7 +955,7 @@ void FAutomationControllerManager::HandlePerformanceDataRequest(const FAutomatio
 	MessageEndpoint->Send(ResponseMessage, Context->GetSender());
 }
 
-void FAutomationControllerManager::HandleRequestNextNetworkCommandMessage(const FAutomationWorkerRequestNextNetworkCommand& Message, const IMessageContextRef& Context)
+void FAutomationControllerManager::HandleRequestNextNetworkCommandMessage(const FAutomationWorkerRequestNextNetworkCommand& Message, const TSharedRef<IMessageContext, ESPMode::ThreadSafe>& Context)
 {
 	// Harvest iteration of running the tests this result came from (stops stale results from being committed to subsequent runs)
 	if ( Message.ExecutionCount == ExecutionCount )
@@ -993,7 +993,7 @@ void FAutomationControllerManager::HandleRequestNextNetworkCommandMessage(const
 	}
 }
 
-void FAutomationControllerManager::HandleRequestTestsReplyCompleteMessage(const FAutomationWorkerRequestTestsReplyComplete& Message, const IMessageContextRef& Context)
+void FAutomationControllerManager::HandleRequestTestsReplyCompleteMessage(const FAutomationWorkerRequestTestsReplyComplete& Message, const TSharedRef<IMessageContext, ESPMode::ThreadSafe>& Context)
 {
 	TArray<FAutomationTestInfo> TestInfo;
 	TestInfo.Reset(Message.Tests.Num());
@@ -1006,7 +1006,7 @@ void FAutomationControllerManager::HandleRequestTestsReplyCompleteMessage(const
 	SetTestNames(Context->GetSender(), TestInfo);
 }
 
-void FAutomationControllerManager::HandleRunTestsReplyMessage(const FAutomationWorkerRunTestsReply& Message, const IMessageContextRef& Context)
+void FAutomationControllerManager::HandleRunTestsReplyMessage(const FAutomationWorkerRunTestsReply& Message, const TSharedRef<IMessageContext, ESPMode::ThreadSafe>& Context)
 {
 	// If we should commit these results
 	if ( Message.ExecutionCount == ExecutionCount )
@@ -1095,7 +1095,7 @@ void FAutomationControllerManager::HandleRunTestsReplyMessage(const FAutomationW
 	RemoveTestRunning(Context->GetSender());
 }
 
-void FAutomationControllerManager::HandleWorkerOfflineMessage( const FAutomationWorkerWorkerOffline& Message, const IMessageContextRef& Context )
+void FAutomationControllerManager::HandleWorkerOfflineMessage( const FAutomationWorkerWorkerOffline& Message, const TSharedRef<IMessageContext, ESPMode::ThreadSafe>& Context )
 {
 	FMessageAddress DeviceMessageAddress = Context->GetSender();
 	DeviceClusterManager.Remove(DeviceMessageAddress);
diff --git a/Engine/Source/Developer/AutomationController/Private/AutomationControllerModule.h b/Engine/Source/Developer/AutomationController/Private/AutomationControllerModule.h
--- a/Engine/Source/Developer/AutomationController/Private/AutomationControllerModule.h
+++ b/Engine/Source/Developer/AutomationController/Private/AutomationControllerModule.h
@@ -3,7 +3,7 @@
 #pragma once
 
 #include "CoreMinimal.h"
-#include "Interfaces/IAutomationControllerModule.h"
+#include "IAutomationControllerModule.h"
 
 /**
  * Implements the AutomationController module.
diff --git a/Engine/Source/Developer/AutomationController/Private/AutomationDeviceClusterManager.cpp b/Engine/Source/Developer/AutomationController/Private/AutomationDeviceClusterManager.cpp
--- a/Engine/Source/Developer/AutomationController/Private/AutomationDeviceClusterManager.cpp
+++ b/Engine/Source/Developer/AutomationController/Private/AutomationDeviceClusterManager.cpp
@@ -1,7 +1,7 @@
 // Copyright 1998-2017 Epic Games, Inc. All Rights Reserved.
 
 #include "AutomationDeviceClusterManager.h"
-#include "Interfaces/IAutomationControllerManager.h"
+#include "IAutomationControllerManager.h"
 
 void FAutomationDeviceClusterManager::Reset()
 {
diff --git a/Engine/Source/Developer/AutomationController/Private/AutomationDeviceClusterManager.h b/Engine/Source/Developer/AutomationController/Private/AutomationDeviceClusterManager.h
--- a/Engine/Source/Developer/AutomationController/Private/AutomationDeviceClusterManager.h
+++ b/Engine/Source/Developer/AutomationController/Private/AutomationDeviceClusterManager.h
@@ -3,7 +3,7 @@
 #pragma once
 
 #include "CoreMinimal.h"
-#include "Interfaces/IAutomationReport.h"
+#include "IAutomationReport.h"
 #include "AutomationWorkerMessages.h"
 #include "IMessageContext.h"
 
diff --git a/Engine/Source/Developer/AutomationController/Private/AutomationReport.h b/Engine/Source/Developer/AutomationController/Private/AutomationReport.h
--- a/Engine/Source/Developer/AutomationController/Private/AutomationReport.h
+++ b/Engine/Source/Developer/AutomationController/Private/AutomationReport.h
@@ -4,7 +4,7 @@
 
 #include "CoreMinimal.h"
 #include "Misc/AutomationTest.h"
-#include "Interfaces/IAutomationReport.h"
+#include "IAutomationReport.h"
 
 /**
  * Interface for automation test results modules.
diff --git a/Engine/Source/Developer/AutomationWindow/Private/AutomationWindowModule.cpp b/Engine/Source/Developer/AutomationWindow/Private/AutomationWindowModule.cpp
--- a/Engine/Source/Developer/AutomationWindow/Private/AutomationWindowModule.cpp
+++ b/Engine/Source/Developer/AutomationWindow/Private/AutomationWindowModule.cpp
@@ -3,10 +3,12 @@
 #include "CoreMinimal.h"
 #include "Modules/ModuleManager.h"
 #include "Widgets/DeclarativeSyntaxSupport.h"
-#include "Interfaces/IAutomationWindowModule.h"
 #include "ISessionManager.h"
 #include "SAutomationWindow.h"
 
+#include "IAutomationWindowModule.h"
+
+
 /**
  * Implements the AutomationWindow module.
  */
diff --git a/Engine/Source/Developer/AutomationWindow/Private/SAutomationExportMenu.h b/Engine/Source/Developer/AutomationWindow/Private/SAutomationExportMenu.h
--- a/Engine/Source/Developer/AutomationWindow/Private/SAutomationExportMenu.h
+++ b/Engine/Source/Developer/AutomationWindow/Private/SAutomationExportMenu.h
@@ -7,7 +7,7 @@
 #include "Input/Reply.h"
 #include "Widgets/DeclarativeSyntaxSupport.h"
 #include "Widgets/SCompoundWidget.h"
-#include "Interfaces/IAutomationControllerManager.h"
+#include "IAutomationControllerManager.h"
 
 class SButton;
 class SComboButton;
diff --git a/Engine/Source/Developer/AutomationWindow/Private/SAutomationTestItem.cpp b/Engine/Source/Developer/AutomationWindow/Private/SAutomationTestItem.cpp
--- a/Engine/Source/Developer/AutomationWindow/Private/SAutomationTestItem.cpp
+++ b/Engine/Source/Developer/AutomationWindow/Private/SAutomationTestItem.cpp
@@ -2,8 +2,8 @@
 
 #include "SAutomationTestItem.h"
 #include "Modules/ModuleManager.h"
-#include "Interfaces/IAutomationReport.h"
-#include "Interfaces/IAutomationControllerModule.h"
+#include "IAutomationReport.h"
+#include "IAutomationControllerModule.h"
 #include "SlateOptMacros.h"
 #include "Framework/Application/SlateApplication.h"
 #include "Widgets/Images/SImage.h"
diff --git a/Engine/Source/Developer/CrashDebugHelper/Private/CrashDebugHelper.cpp b/Engine/Source/Developer/CrashDebugHelper/Private/CrashDebugHelper.cpp
--- a/Engine/Source/Developer/CrashDebugHelper/Private/CrashDebugHelper.cpp
+++ b/Engine/Source/Developer/CrashDebugHelper/Private/CrashDebugHelper.cpp
@@ -236,7 +236,7 @@ bool ICrashDebugHelper::Init()
 		if (CrashInfo.bMutexPDBCache && !CrashInfo.PDBCacheLockName.IsEmpty())
 		{
 			// Scoped lock
-			FSystemWideCriticalSection PDBCacheLock(CrashInfo.PDBCacheLockName, FTimespan(0, 0, 10, 0, 0));
+			FSystemWideCriticalSection PDBCacheLock(CrashInfo.PDBCacheLockName, FTimespan::FromMinutes(10.0));
 			if (PDBCacheLock.IsValid())
 			{
 				FPDBCache::Get().Init();
diff --git a/Engine/Source/Developer/CrashDebugHelper/Private/Windows/CrashDebugHelperWindows.cpp b/Engine/Source/Developer/CrashDebugHelper/Private/Windows/CrashDebugHelperWindows.cpp
--- a/Engine/Source/Developer/CrashDebugHelper/Private/Windows/CrashDebugHelperWindows.cpp
+++ b/Engine/Source/Developer/CrashDebugHelper/Private/Windows/CrashDebugHelperWindows.cpp
@@ -44,7 +44,7 @@ bool FCrashDebugHelperWindows::CreateMinidumpDiagnosticReport( const FString& In
 			{
 				// Scoped lock
 				UE_LOG(LogCrashDebugHelper, Log, TEXT("Locking for InitSymbols()"));
-				FSystemWideCriticalSection PDBCacheLock(CrashInfo.PDBCacheLockName, FTimespan(0, 0, 10, 0, 0));
+				FSystemWideCriticalSection PDBCacheLock(CrashInfo.PDBCacheLockName, FTimespan::FromMinutes(10.0));
 				if (PDBCacheLock.IsValid())
 				{
 					bInitSymbols = InitSymbols(WindowsStackWalkExt, bSyncSymbols);
@@ -84,7 +84,7 @@ bool FCrashDebugHelperWindows::CreateMinidumpDiagnosticReport( const FString& In
 					{
 						// Scoped lock
 						UE_LOG(LogCrashDebugHelper, Log, TEXT("Locking for SyncAndReadSourceFile()"));
-						const FTimespan GlobalLockWaitTimeout(0, 0, 0, 30, 0);
+						const FTimespan GlobalLockWaitTimeout = FTimespan::FromSeconds(30.0);
 						FSystemWideCriticalSection SyncSourceLock(SourceSyncLockName, GlobalLockWaitTimeout);
 						if (SyncSourceLock.IsValid())
 						{
diff --git a/Engine/Source/Developer/DeviceManager/Private/DeviceManagerModule.cpp b/Engine/Source/Developer/DeviceManager/Private/DeviceManagerModule.cpp
--- a/Engine/Source/Developer/DeviceManager/Private/DeviceManagerModule.cpp
+++ b/Engine/Source/Developer/DeviceManager/Private/DeviceManagerModule.cpp
@@ -1,19 +1,20 @@
 // Copyright 1998-2017 Epic Games, Inc. All Rights Reserved.
 
-#include "CoreMinimal.h"
+#include "CoreTypes.h"
+#include "EditorStyleSet.h"
+#include "Framework/Docking/TabManager.h"
+#include "IDeviceManagerModule.h"
+#include "ITargetDeviceServicesModule.h"
 #include "Modules/ModuleManager.h"
-#include "Widgets/DeclarativeSyntaxSupport.h"
+#include "Templates/SharedPointer.h"
 #include "Textures/SlateIcon.h"
-#include "Widgets/SWidget.h"
-#include "Framework/Docking/TabManager.h"
-#include "EditorStyleSet.h"
-#include "Interfaces/ITargetDeviceServicesModule.h"
-#include "Interfaces/IDeviceManagerModule.h"
+#include "Widgets/DeclarativeSyntaxSupport.h"
 #include "Widgets/SDeviceManager.h"
-
+#include "Widgets/SWidget.h"
+#include "Widgets/Docking/SDockTab.h"
 #include "WorkspaceMenuStructure.h"
 #include "WorkspaceMenuStructureModule.h"
-#include "Widgets/Docking/SDockTab.h"
+
 
 static const FName DeviceManagerTabName("DeviceManager");
 
@@ -26,9 +27,9 @@ class FDeviceManagerModule
 {
 public:
 
-	// IModuleInterface interface
+	//~ IModuleInterface interface
 
-	virtual void StartupModule( ) override
+	virtual void StartupModule() override
 	{
 		// @todo gmp: implement an IoC container
 		ITargetDeviceServicesModule& TargetDeviceServicesModule = FModuleManager::LoadModuleChecked<ITargetDeviceServicesModule>(TEXT("TargetDeviceServices"));
@@ -47,16 +48,16 @@ public:
 #endif
 	}
 
-	virtual void ShutdownModule( ) override
+	virtual void ShutdownModule() override
 	{
 		FGlobalTabmanager::Get()->UnregisterNomadTabSpawner(DeviceManagerTabName);
 	}
 
 public:
 
-	// IDeviceManagerModule interface
+	//~ IDeviceManagerModule interface
 
-	virtual TSharedRef<SWidget> CreateDeviceManager( const ITargetDeviceServiceManagerRef& DeviceServiceManager, const TSharedRef<SDockTab>& ConstructUnderMajorTab, const TSharedPtr<SWindow>& ConstructUnderWindow ) override
+	virtual TSharedRef<SWidget> CreateDeviceManager(const TSharedRef<ITargetDeviceServiceManager>& DeviceServiceManager, const TSharedRef<SDockTab>& ConstructUnderMajorTab, const TSharedPtr<SWindow>& ConstructUnderWindow) override
 	{
 		return SNew(SDeviceManager, DeviceServiceManager, ConstructUnderMajorTab, ConstructUnderWindow);
 	}
@@ -69,7 +70,7 @@ private:
 	 * @param SpawnTabArgs The arguments for the tab to spawn.
 	 * @return The spawned tab.
 	 */
-	TSharedRef<SDockTab> SpawnDeviceManagerTab( const FSpawnTabArgs& SpawnTabArgs )
+	TSharedRef<SDockTab> SpawnDeviceManagerTab(const FSpawnTabArgs& SpawnTabArgs)
 	{
 		const TSharedRef<SDockTab> DockTab = SNew(SDockTab)
 			.TabRole(ETabRole::MajorTab);
@@ -82,7 +83,7 @@ private:
 private:
 
 	// @todo gmp: implement an IoC container
-	ITargetDeviceServiceManagerPtr TargetDeviceServiceManager;
+	TSharedPtr<ITargetDeviceServiceManager> TargetDeviceServiceManager;
 };
 
 
diff --git a/Engine/Source/Developer/DeviceManager/Private/Models/DeviceBrowserFilter.h b/Engine/Source/Developer/DeviceManager/Private/Models/DeviceBrowserFilter.h
--- a/Engine/Source/Developer/DeviceManager/Private/Models/DeviceBrowserFilter.h
+++ b/Engine/Source/Developer/DeviceManager/Private/Models/DeviceBrowserFilter.h
@@ -2,29 +2,31 @@
 
 #pragma once
 
-#include "CoreMinimal.h"
-#include "Interfaces/ITargetDeviceService.h"
+#include "CoreTypes.h"
+#include "Containers/Array.h"
+#include "Containers/Map.h"
+#include "Containers/UnrealString.h"
+#include "Delegates/Delegate.h"
+#include "Internationalization/Text.h"
+#include "ITargetDeviceService.h"
+#include "Templates/SharedPointer.h"
 
-class FDeviceBrowserFilter;
 
-/** Type definition for shared pointers to instances of FDeviceBrowserFilter. */
-typedef TSharedPtr<class FDeviceBrowserFilter> FDeviceBrowserFilterPtr;
-
-/** Type definition for shared references to instances of FDeviceBrowserFilter. */
-typedef TSharedRef<class FDeviceBrowserFilter> FDeviceBrowserFilterRef;
-
-class FDeviceBrowserFilterEntry
+/**
+ * An entry in the device browser filter.
+ */
+struct FDeviceBrowserFilterEntry
 {
-public:
+	FString PlatformName;
+	FName PlatformLookup;
+
 	FDeviceBrowserFilterEntry(FString InPlatformName, FName InPlatformLookup)
 		: PlatformName(InPlatformName)
 		, PlatformLookup(InPlatformLookup)
-	{
-	}
-	FString PlatformName;
-	FName PlatformLookup;
+	{ }
 };
 
+
 /**
  * Implements a filter for the device browser's target device service list.
  */
@@ -33,12 +35,12 @@ class FDeviceBrowserFilter
 public:
 
 	/**
-	 * Filters the specified target device service based on the current filter settings.
+	 * Filter the specified target device service based on the current filter settings.
 	 *
 	 * @param Device The service to filter.
 	 * @return true if the service passed the filter, false otherwise.
 	 */
-	bool FilterDeviceService( const ITargetDeviceServicePtr& DeviceService )
+	bool FilterDeviceService(const TSharedPtr<ITargetDeviceService, ESPMode::ThreadSafe>& DeviceService)
 	{
 		if (DeviceService.IsValid())
 		{
@@ -52,35 +54,35 @@ public:
 				return true;
 			}
 
-			return (DeviceService->GetDeviceName().Contains(DeviceSearchText.ToString()) );
+			return (DeviceService->GetDeviceName().Contains(DeviceSearchText.ToString()));
 		}
 
 		return false;
 	}
 
 	/**
-	 * Gets the number of target device services that have the specified platform.
+	 * Get the number of target device services that have the specified platform.
 	 *
 	 * @param PlatformName The name of the platform.
 	 * @return The number of matching services.
 	 */
-	int32 GetServiceCountPerPlatform( const FString& PlatformName ) const
+	int32 GetServiceCountPerPlatform(const FString& PlatformName) const
 	{
 		return PlatformCounters.FindRef(PlatformName);
 	}
 
 	/**
-	 * Gets the device search string.
+	 * Get the device search string.
 	 *
 	 * @return Search string.
 	 */
-	const FText& GetDeviceSearchText( ) const
+	const FText& GetDeviceSearchText() const
 	{
 		return DeviceSearchText;
 	}
 
 	/**
-	 * Returns the list of filtered platforms.
+	 * Return the list of filtered platforms.
 	 *
 	 * @return List of platform names.
 	 */
@@ -90,22 +92,23 @@ public:
 	}
 
 	/**
-	 * Checks whether the specified platform is enabled in the filter.
+	 * Check whether the specified platform is enabled in the filter.
 	 *
 	 * @param PlatformName The name of the platform to check.
 	 * @return true if the platform is enabled, false otherwise.
+	 * @see SetPlatformEnabled
 	 */
-	bool IsPlatformEnabled( const FString& PlatformName ) const
+	bool IsPlatformEnabled(const FString& PlatformName) const
 	{
 		return !DisabledPlatforms.Contains(PlatformName);
 	}
 
 	/**
-	 * Populates the filter from the given list of target device proxies.
+	 * Populate the filter from the given list of target device proxies.
 	 *
 	 * @param DeviceProxies The list of device proxies to populate the filter from.
 	 */
-	void ResetFilter( const TArray<ITargetDeviceServicePtr>& DeviceServices )
+	void ResetFilter(const TArray<ITargetDeviceServicePtr>& DeviceServices)
 	{
 		PlatformList.Reset();
 		PlatformCounters.Reset();
@@ -134,7 +137,7 @@ public:
 	 *
 	 * @param SearchText The search string.
 	 */
-	void SetDeviceSearchString( const FText& SearchText )
+	void SetDeviceSearchString(const FText& SearchText)
 	{
 		if (!DeviceSearchText.EqualTo(SearchText))
 		{
@@ -145,12 +148,13 @@ public:
 	}
 
 	/**
-	 * Sets the enabled state of the specified device proxy platform.
+	 * Set the enabled state of the specified device proxy platform.
 	 *
 	 * @param PlatformName The name of the platform to enable or disable in the filter.
 	 * @param Enabled Whether the platform is enabled.
+	 * @see IsPlatformEnabled
 	 */
-	void SetPlatformEnabled( const FString& PlatformName, bool Enabled )
+	void SetPlatformEnabled(const FString& PlatformName, bool Enabled)
 	{
 		if (Enabled)
 		{
@@ -167,49 +171,51 @@ public:
 public:
 
 	/**
-	 * Gets a delegate to be invoked when the filter state changed.
+	 * Get a delegate to be invoked when the filter state changed.
 	 *
 	 * @return The delegate.
+	 * @see OnFilterReset
 	 */
 	DECLARE_EVENT(FDeviceBrowserFilter, FOnDeviceBrowserFilterChanged);
-	FOnDeviceBrowserFilterChanged& OnFilterChanged( )
+	FOnDeviceBrowserFilterChanged& OnFilterChanged()
 	{
 		return FilterChangedEvent;
 	}
 
 	/**
-	 * Gets a delegate to be invoked when the filter has been reset.
+	 * Get a delegate to be invoked when the filter has been reset.
 	 *
 	 * @return The delegate.
+	 * @see OnFilterChanged
 	 */
 	DECLARE_EVENT(FDeviceBrowserFilter, FOnDeviceBrowserFilterReset);
-	FOnDeviceBrowserFilterReset& OnFilterReset( )
+	FOnDeviceBrowserFilterReset& OnFilterReset()
 	{
 		return FilterResetEvent;
 	}
 
 private:
 
-	// Holds the device search string.
+	/** The device search string. */
 	FText DeviceSearchText;
 
-	// Holds the list of disabled platforms.
+	/** The list of disabled platforms. */
 	TArray<FString> DisabledPlatforms;
 
-	// Holds the device counters for owner filters.
+	/** The device counters for owner filters. */
 	TMap<FString, int32> OwnerCounters;
 
-	// Holds the device counters for platform filters.
+	/** The device counters for platform filters. */
 	TMap<FString, int32> PlatformCounters;
 
-	// Holds the list of platform filters.
+	/** The list of platform filters. */
 	TArray<TSharedPtr<FDeviceBrowserFilterEntry>> PlatformList;
 
 private:
 
-	// Holds an event delegate that is invoked when the filter state changed.
+	/** An event delegate that is invoked when the filter state changed. */
 	FOnDeviceBrowserFilterChanged FilterChangedEvent;
 
-	// Holds an event delegate that is invoked when the filter has been reset.
+	/** An event delegate that is invoked when the filter has been reset. */
 	FOnDeviceBrowserFilterReset FilterResetEvent;
 };
diff --git a/Engine/Source/Developer/DeviceManager/Private/Models/DeviceDetailsCommands.h b/Engine/Source/Developer/DeviceManager/Private/Models/DeviceDetailsCommands.h
--- a/Engine/Source/Developer/DeviceManager/Private/Models/DeviceDetailsCommands.h
+++ b/Engine/Source/Developer/DeviceManager/Private/Models/DeviceDetailsCommands.h
@@ -2,13 +2,17 @@
 
 #pragma once
 
-#include "CoreMinimal.h"
-#include "Framework/Commands/InputChord.h"
+#include "CoreTypes.h"
 #include "EditorStyleSet.h"
 #include "Framework/Commands/Commands.h"
+#include "Framework/Commands/InputChord.h"
+#include "Internationalization/Text.h"
+#include "Templates/SharedPointer.h"
+
 
 #define LOCTEXT_NAMESPACE "DeviceDetailsCommands"
 
+
 /**
  * The device details commands
  */
@@ -17,9 +21,7 @@ class FDeviceDetailsCommands
 {
 public:
 
-	/**
-	 * Default constructor.
-	 */
+	/** Default constructor. */
 	FDeviceDetailsCommands()
 		: TCommands<FDeviceDetailsCommands>(
 			"DeviceDetails",
@@ -30,9 +32,9 @@ public:
 
 public:
 
-	// TCommands interface
+	//~ TCommands interface
 
-	virtual void RegisterCommands( ) override
+	virtual void RegisterCommands() override
 	{
 		UI_COMMAND(Claim, "Claim", "Claim the device", EUserInterfaceActionType::Button, FInputChord());
 		UI_COMMAND(Release, "Release", "Release the device", EUserInterfaceActionType::Button, FInputChord());
diff --git a/Engine/Source/Developer/DeviceManager/Private/Models/DeviceDetailsFeature.h b/Engine/Source/Developer/DeviceManager/Private/Models/DeviceDetailsFeature.h
--- a/Engine/Source/Developer/DeviceManager/Private/Models/DeviceDetailsFeature.h
+++ b/Engine/Source/Developer/DeviceManager/Private/Models/DeviceDetailsFeature.h
@@ -2,15 +2,8 @@
 
 #pragma once
 
-#include "CoreMinimal.h"
-
-struct FDeviceDetailsFeature;
-
-/** Type definition for shared pointers to instances of FDeviceDetailsFeature. */
-typedef TSharedPtr<struct FDeviceDetailsFeature> FDeviceDetailsFeaturePtr;
-
-/** Type definition for shared references to instances of FDeviceDetailsFeature. */
-typedef TSharedRef<struct FDeviceDetailsFeature> FDeviceDetailsFeatureRef;
+#include "CoreTypes.h"
+#include "Containers/UnrealString.h"
 
 
 /**
@@ -18,19 +11,19 @@ typedef TSharedRef<struct FDeviceDetailsFeature> FDeviceDetailsFeatureRef;
  */
 struct FDeviceDetailsFeature
 {
-	/** Holds the name of the feature. */
+	/** The name of the feature. */
 	FString FeatureName;
 
 	/** Whether the feature is available. */
 	bool Available;
 
 	/**
-	 * Creates and initializes a new instance.
+	 * Create and initialize a new instance.
 	 *
 	 * @param InFeatureName The name of the feature.
 	 * @param InAvailable Whether the feature is available.
 	 */
-	FDeviceDetailsFeature( const FString& InFeatureName, bool InAvailable )
+	FDeviceDetailsFeature(const FString& InFeatureName, bool InAvailable)
 		: FeatureName(InFeatureName)
 		, Available(InAvailable)
 	{ }
diff --git a/Engine/Source/Developer/DeviceManager/Private/Models/DeviceManagerModel.h b/Engine/Source/Developer/DeviceManager/Private/Models/DeviceManagerModel.h
--- a/Engine/Source/Developer/DeviceManager/Private/Models/DeviceManagerModel.h
+++ b/Engine/Source/Developer/DeviceManager/Private/Models/DeviceManagerModel.h
@@ -2,17 +2,13 @@
 
 #pragma once
 
-#include "CoreMinimal.h"
-#include "Interfaces/ITargetDeviceService.h"
+#include "CoreTypes.h"
+#include "Delegates/Delegate.h"
+#include "ITargetDeviceService.h"
+#include "Templates/SharedPointer.h"
 
 class FDeviceManagerModel;
 
-/** Type definition for shared pointers to instances of FDeviceManagerModel. */
-typedef TSharedPtr<class FDeviceManagerModel> FDeviceManagerModelPtr;
-
-/** Type definition for shared references to instances of FDeviceManagerModel. */
-typedef TSharedRef<class FDeviceManagerModel> FDeviceManagerModelRef;
-
 
 /**
  * Implements a view model for the messaging debugger.
@@ -22,21 +18,21 @@ class FDeviceManagerModel
 public:
 
 	/**
-	 * Gets the device service that is currently selected in the device browser.
+	 * Get the device service that is currently selected in the device browser.
 	 *
 	 * @return The selected device service.
 	 */
-	ITargetDeviceServicePtr GetSelectedDeviceService( ) const
+	TSharedPtr<ITargetDeviceService, ESPMode::ThreadSafe> GetSelectedDeviceService() const
 	{
 		return SelectedDeviceService;
 	}
 
 	/**
-	 * Selects the specified device service (or none if nullptr).
+	 * Select the specified device service (or none if nullptr).
 	 *
 	 * @param DeviceService The device service to select.
 	 */
-	void SelectDeviceService( const ITargetDeviceServicePtr& DeviceService )
+	void SelectDeviceService(const TSharedPtr<ITargetDeviceService, ESPMode::ThreadSafe>& DeviceService)
 	{
 		if (SelectedDeviceService != DeviceService)
 		{
@@ -48,23 +44,21 @@ public:
 public:
 
 	/**
-	 * Gets an event delegate that is invoked when the selected device service has changed.
+	 * Get an event delegate that is invoked when the selected device service has changed.
 	 *
 	 * @return The event delegate.
 	 */
 	DECLARE_EVENT(FDeviceManagerModel, FOnSelectedDeviceServiceChanged);
-	FOnSelectedDeviceServiceChanged& OnSelectedDeviceServiceChanged( )
+	FOnSelectedDeviceServiceChanged& OnSelectedDeviceServiceChanged()
 	{
 		return SelectedDeviceServiceChangedEvent;
 	}
 
 private:
 
-	// Holds the currently selected target device service.
-	ITargetDeviceServicePtr SelectedDeviceService;
-
-private:
+	/** The currently selected target device service. */
+	TSharedPtr<ITargetDeviceService, ESPMode::ThreadSafe> SelectedDeviceService;
 
-	// Holds an event delegate that is invoked when the selected message has changed.
+	/** An event delegate that is invoked when the selected message has changed. */
 	FOnSelectedDeviceServiceChanged SelectedDeviceServiceChangedEvent;
 };
diff --git a/Engine/Source/Developer/DeviceManager/Private/Widgets/Apps/SDeviceApps.cpp b/Engine/Source/Developer/DeviceManager/Private/Widgets/Apps/SDeviceApps.cpp
--- a/Engine/Source/Developer/DeviceManager/Private/Widgets/Apps/SDeviceApps.cpp
+++ b/Engine/Source/Developer/DeviceManager/Private/Widgets/Apps/SDeviceApps.cpp
@@ -1,11 +1,13 @@
 // Copyright 1998-2017 Epic Games, Inc. All Rights Reserved.
 
-#include "Widgets/Apps/SDeviceApps.h"
-#include "Widgets/SOverlay.h"
+#include "SDeviceApps.h"
+
+#include "EditorStyleSet.h"
 #include "SlateOptMacros.h"
+#include "Widgets/SOverlay.h"
 #include "Widgets/Text/STextBlock.h"
-#include "EditorStyleSet.h"
-#include "Widgets/Views/SListView.h"
+
+#include "Models/DeviceManagerModel.h"
 #include "Widgets/Apps/SDeviceAppsAppListRow.h"
 
 
@@ -15,7 +17,7 @@
 /* SMessagingEndpoints structors
 *****************************************************************************/
 
-SDeviceApps::~SDeviceApps( )
+SDeviceApps::~SDeviceApps()
 {
 	if (Model.IsValid())
 	{
@@ -28,10 +30,37 @@ SDeviceApps::~SDeviceApps( )
  *****************************************************************************/
 
 BEGIN_SLATE_FUNCTION_BUILD_OPTIMIZATION
-void SDeviceApps::Construct( const FArguments& InArgs, const FDeviceManagerModelRef& InModel )
+void SDeviceApps::Construct(const FArguments& InArgs, const TSharedRef<FDeviceManagerModel>& InModel)
 {
 	Model = InModel;
 
+	// callback for getting the enabled state of the processes panel
+	auto AppsBoxIsEnabled = [this]() -> bool {
+		ITargetDeviceServicePtr DeviceService = Model->GetSelectedDeviceService();
+		return (DeviceService.IsValid() && DeviceService->CanStart());
+	};
+
+	// callback for generating a row widget in the application list view
+	auto AppListViewGenerateRow = [this](TSharedPtr<FString> Item, const TSharedRef<STableViewBase>& OwnerTable) -> TSharedRef<ITableRow> {
+		return SNew(SDeviceAppsAppListRow, OwnerTable);
+	};
+
+	// callback for selecting items in the devices list view
+	auto AppListViewSelectionChanged = [this](TSharedPtr<FString> Selection, ESelectInfo::Type SelectInfo) {
+		// @todo
+	};
+
+	// callback for getting the visibility of the 'Select a device' message. */
+	auto SelectDeviceOverlayVisibility = [this]() -> EVisibility {
+		if (Model->GetSelectedDeviceService().IsValid())
+		{
+			return EVisibility::Hidden;
+		}
+
+		return EVisibility::Visible;
+	};
+
+	// construct children
 	ChildSlot
 	[
 		SNew(SOverlay)
@@ -39,7 +68,7 @@ void SDeviceApps::Construct( const FArguments& InArgs, const FDeviceManagerModel
 		+ SOverlay::Slot()
 			[
 				SNew(SVerticalBox)
-					.IsEnabled(this, &SDeviceApps::HandleAppsBoxIsEnabled)
+					.IsEnabled_Lambda(AppsBoxIsEnabled)
 
 				+ SVerticalBox::Slot()
 					.FillHeight(1.0f)
@@ -53,8 +82,8 @@ void SDeviceApps::Construct( const FArguments& InArgs, const FDeviceManagerModel
 								SAssignNew(AppListView, SListView<TSharedPtr<FString> >)
 									.ItemHeight(20.0f)
 									.ListItemsSource(&AppList)
-									.OnGenerateRow(this, &SDeviceApps::HandleAppListViewGenerateRow)
-									.OnSelectionChanged(this, &SDeviceApps::HandleAppListViewSelectionChanged)
+									.OnGenerateRow_Lambda(AppListViewGenerateRow)
+									.OnSelectionChanged_Lambda(AppListViewSelectionChanged)
 									.SelectionMode(ESelectionMode::Single)
 									.HeaderRow
 									(
@@ -80,7 +109,7 @@ void SDeviceApps::Construct( const FArguments& InArgs, const FDeviceManagerModel
 				SNew(SBorder)
 					.BorderImage(FEditorStyle::GetBrush("NotificationList.ItemBackground"))
 					.Padding(8.0f)
-					.Visibility(this, &SDeviceApps::HandleSelectDeviceOverlayVisibility)
+					.Visibility_Lambda(SelectDeviceOverlayVisibility)
 					[
 						SNew(STextBlock)
 							.Text(LOCTEXT("SelectSessionOverlayText", "Please select a device from the Device Browser"))
@@ -88,48 +117,14 @@ void SDeviceApps::Construct( const FArguments& InArgs, const FDeviceManagerModel
 			]
 	];
 
-	Model->OnSelectedDeviceServiceChanged().AddRaw(this, &SDeviceApps::HandleModelSelectedDeviceServiceChanged);
-}
-END_SLATE_FUNCTION_BUILD_OPTIMIZATION
-
-
-/* SDeviceApps callbacks
- *****************************************************************************/
-
-TSharedRef<ITableRow> SDeviceApps::HandleAppListViewGenerateRow( TSharedPtr<FString> Item, const TSharedRef<STableViewBase>& OwnerTable )
-{
-	return SNew(SDeviceAppsAppListRow, OwnerTable);
-}
-
-
-void SDeviceApps::HandleAppListViewSelectionChanged( TSharedPtr<FString> Selection, ESelectInfo::Type SelectInfo )
-{
-
-}
-
-
-bool SDeviceApps::HandleAppsBoxIsEnabled( ) const
-{
-	ITargetDeviceServicePtr DeviceService = Model->GetSelectedDeviceService();
+	// callback for handling device service selection changes
+	auto ModelSelectedDeviceServiceChanged = [this]() {
+		// @todo
+	};
 
-	return (DeviceService.IsValid() && DeviceService->CanStart());
-}
-
-
-void SDeviceApps::HandleModelSelectedDeviceServiceChanged( )
-{
-}
-
-
-EVisibility SDeviceApps::HandleSelectDeviceOverlayVisibility( ) const
-{
-	if (Model->GetSelectedDeviceService().IsValid())
-	{
-		return EVisibility::Hidden;
-	}
-
-	return EVisibility::Visible;
+	Model->OnSelectedDeviceServiceChanged().AddLambda(ModelSelectedDeviceServiceChanged);
 }
+END_SLATE_FUNCTION_BUILD_OPTIMIZATION
 
 
 #undef LOCTEXT_NAMESPACE
diff --git a/Engine/Source/Developer/DeviceManager/Private/Widgets/Apps/SDeviceApps.h b/Engine/Source/Developer/DeviceManager/Private/Widgets/Apps/SDeviceApps.h
--- a/Engine/Source/Developer/DeviceManager/Private/Widgets/Apps/SDeviceApps.h
+++ b/Engine/Source/Developer/DeviceManager/Private/Widgets/Apps/SDeviceApps.h
@@ -2,13 +2,15 @@
 
 #pragma once
 
-#include "CoreMinimal.h"
-#include "Layout/Visibility.h"
+#include "CoreTypes.h"
+#include "Containers/Array.h"
+#include "Templates/SharedPointer.h"
 #include "Widgets/DeclarativeSyntaxSupport.h"
 #include "Widgets/SCompoundWidget.h"
-#include "Widgets/Views/STableViewBase.h"
-#include "Widgets/Views/STableRow.h"
-#include "Models/DeviceManagerModel.h"
+#include "Widgets/Views/SListView.h"
+
+class FDeviceManagerModel;
+
 
 /**
  * Implements the device details widget.
@@ -23,10 +25,8 @@ public:
 
 public:
 
-	/**
-	 * Destructor.
-	 */
-	~SDeviceApps( );
+	/** Destructor. */
+	~SDeviceApps();
 
 public:
 
@@ -36,33 +36,16 @@ public:
 	 * @param InArgs The construction arguments.
 	 * @param InModel The view model to use.
 	 */
-	void Construct( const FArguments& InArgs, const FDeviceManagerModelRef& InModel );
-
-private:
-
-	// Callback for generating a row widget in the application list view.
-	TSharedRef<ITableRow> HandleAppListViewGenerateRow( TSharedPtr<FString> Item, const TSharedRef<STableViewBase>& OwnerTable );
-
-	// Callback for selecting items in the devices list view.
-	void HandleAppListViewSelectionChanged( TSharedPtr<FString> Selection, ESelectInfo::Type SelectInfo );
-
-	// Callback for getting the enabled state of the processes panel.
-	bool HandleAppsBoxIsEnabled( ) const;
-
-	// Callback for handling device service selection changes.
-	void HandleModelSelectedDeviceServiceChanged( );
-
-	// Callback for getting the visibility of the 'Select a device' message.
-	EVisibility HandleSelectDeviceOverlayVisibility( ) const;
-
+	void Construct(const FArguments& InArgs, const TSharedRef<FDeviceManagerModel>& InModel);
+	
 private:
 
-	// Holds the list of applications deployed to the device.
-	TArray<TSharedPtr<FString> > AppList;
+	/** The list of applications deployed to the device. */
+	TArray<TSharedPtr<FString>> AppList;
 
-	// Holds the application list view.
+	/** The application list view. */
 	TSharedPtr<SListView<TSharedPtr<FString>>> AppListView;
 
-	// Holds a pointer the device manager's view model.
-	FDeviceManagerModelPtr Model;
+	/** Pointer the device manager's view model. */
+	TSharedPtr<FDeviceManagerModel> Model;
 };
