Copying //UE4/Dev-Sequencer to //UE4/Main (Source: //UE4/Dev-Sequencer @ 3617642)

diff --git a/Engine/Source/Developer/DeviceManager/Private/Widgets/Apps/SDeviceAppsAppListRow.h b/Engine/Source/Developer/DeviceManager/Private/Widgets/Apps/SDeviceAppsAppListRow.h
--- a/Engine/Source/Developer/DeviceManager/Private/Widgets/Apps/SDeviceAppsAppListRow.h
+++ b/Engine/Source/Developer/DeviceManager/Private/Widgets/Apps/SDeviceAppsAppListRow.h
@@ -2,24 +2,27 @@
 
 #pragma once
 
-#include "CoreMinimal.h"
-#include "Widgets/SNullWidget.h"
+#include "CoreTypes.h"
+#include "Layout/Margin.h"
+#include "Templates/SharedPointer.h"
 #include "Widgets/DeclarativeSyntaxSupport.h"
+#include "Widgets/SNullWidget.h"
 #include "Widgets/SWidget.h"
-#include "Layout/Margin.h"
-#include "Widgets/Views/STableViewBase.h"
 #include "Widgets/Layout/SBox.h"
-#include "Widgets/Views/STableRow.h"
 #include "Widgets/Text/STextBlock.h"
 #include "Widgets/Views/SListView.h"
+#include "Widgets/Views/STableRow.h"
+#include "Widgets/Views/STableViewBase.h"
+
 
 #define LOCTEXT_NAMESPACE "SDeviceAppsAppListRow"
 
+
 /**
  * Implements a row widget for the application list view.
  */
 class SDeviceAppsAppListRow
-	: public SMultiColumnTableRow<TSharedPtr<FString> >
+	: public SMultiColumnTableRow<TSharedPtr<FString>>
 {
 public:
 
@@ -33,7 +36,7 @@ public:
 	 *
 	 * @param InArgs The construction arguments.
 	 */
-	void Construct( const FArguments& InArgs, const TSharedRef<STableViewBase>& InOwnerTableView )
+	void Construct(const FArguments& InArgs, const TSharedRef<STableViewBase>& InOwnerTableView)
 	{
 		SMultiColumnTableRow<TSharedPtr<FString> >::Construct(FSuperRowType::FArguments(), InOwnerTableView);
 	}
@@ -46,7 +49,7 @@ public:
 	 * @param ColumnName The name of the column to generate the widget for.
 	 * @return The widget.
 	 */
-	virtual TSharedRef<SWidget> GenerateWidgetForColumn( const FName& ColumnName ) override
+	virtual TSharedRef<SWidget> GenerateWidgetForColumn(const FName& ColumnName) override
 	{
 		if (ColumnName == TEXT("Date"))
 		{
diff --git a/Engine/Source/Developer/DeviceManager/Private/Widgets/Browser/SDeviceBrowser.cpp b/Engine/Source/Developer/DeviceManager/Private/Widgets/Browser/SDeviceBrowser.cpp
--- a/Engine/Source/Developer/DeviceManager/Private/Widgets/Browser/SDeviceBrowser.cpp
+++ b/Engine/Source/Developer/DeviceManager/Private/Widgets/Browser/SDeviceBrowser.cpp
@@ -1,15 +1,22 @@
 // Copyright 1998-2017 Epic Games, Inc. All Rights Reserved.
 
-#include "Widgets/Browser/SDeviceBrowser.h"
-#include "SlateOptMacros.h"
+#include "SDeviceBrowser.h"
+
 #include "EditorStyleSet.h"
+#include "Framework/Commands/UICommandList.h"
+#include "ITargetDeviceServiceManager.h"
+#include "SlateOptMacros.h"
+#include "Widgets/Layout/SExpandableArea.h"
+#include "Widgets/Views/STableRow.h"
+#include "Widgets/Views/STableViewBase.h"
+
 #include "Widgets/Browser/SDeviceBrowserContextMenu.h"
-#include "Widgets/Views/SListView.h"
 #include "Widgets/Browser/SDeviceBrowserDeviceAdder.h"
+#include "Widgets/Browser/SDeviceBrowserDeviceListRow.h"
 #include "Widgets/Browser/SDeviceBrowserFilterBar.h"
 #include "Widgets/Browser/SDeviceBrowserTooltip.h"
-#include "Widgets/Browser/SDeviceBrowserDeviceListRow.h"
-#include "Widgets/Layout/SExpandableArea.h"
+#include "Models/DeviceBrowserFilter.h"
+#include "Models/DeviceManagerModel.h"
 
 
 #define LOCTEXT_NAMESPACE "SDeviceBrowser"
@@ -19,7 +26,7 @@
  *****************************************************************************/
 
 BEGIN_SLATE_FUNCTION_BUILD_OPTIMIZATION
-void SDeviceBrowser::Construct( const FArguments& InArgs, const FDeviceManagerModelRef& InModel, const ITargetDeviceServiceManagerRef& InDeviceServiceManager, const TSharedPtr<FUICommandList>& InUICommandList )
+void SDeviceBrowser::Construct(const FArguments& InArgs, const TSharedRef<FDeviceManagerModel>& InModel, const TSharedRef<ITargetDeviceServiceManager>& InDeviceServiceManager, const TSharedPtr<FUICommandList>& InUICommandList)
 {
 	DeviceServiceManager = InDeviceServiceManager;
 	Filter = MakeShareable(new FDeviceBrowserFilter());
@@ -27,24 +34,50 @@ void SDeviceBrowser::Construct( const FArguments& InArgs, const FDeviceManagerMo
 	NeedsServiceListRefresh = true;
 	UICommandList = InUICommandList;
 
+	auto DeviceServiceListViewContextMenuOpening = [this]() -> TSharedPtr<SWidget> {
+		TArray<TSharedPtr<ITargetDeviceService, ESPMode::ThreadSafe>> SelectedDeviceServices = DeviceServiceListView->GetSelectedItems();
+
+		if (SelectedDeviceServices.Num() > 0)
+		{
+			return SNew(SDeviceBrowserContextMenu, UICommandList);
+		}
+
+		return nullptr;
+	};
+
+	auto DeviceServiceListViewHighlightText = [this]() -> FText {
+		return Filter->GetDeviceSearchText();
+	};
+
+	auto DeviceServiceListViewGenerateRow = [=](TSharedPtr<ITargetDeviceService, ESPMode::ThreadSafe> DeviceService, const TSharedRef<STableViewBase>& OwnerTable) -> TSharedRef<ITableRow> {
+		return SNew(SDeviceBrowserDeviceListRow, OwnerTable)
+			.DeviceService(DeviceService)
+			.HighlightText_Lambda(DeviceServiceListViewHighlightText)
+			.ToolTip(SNew(SDeviceBrowserTooltip, DeviceService.ToSharedRef()));
+	};
+
+	auto DeviceServiceListViewSelectionChanged = [this](TSharedPtr<ITargetDeviceService, ESPMode::ThreadSafe> Selection, ESelectInfo::Type SelectInfo) {
+		Model->SelectDeviceService(Selection);
+	};
+
 	ChildSlot
 	[
 		SNew(SVerticalBox)
 
 		+ SVerticalBox::Slot()
 			.AutoHeight()
-				[
-					/*
-					SNew(SExpandableArea)
-						.AreaTitle(LOCTEXT("FilterBarAreaTitle", "Device Filter").ToString())
-						.InitiallyCollapsed(true)
-						.Padding(FMargin(8.0f, 8.0f, 8.0f, 4.0f))
-						.BodyContent()
-						[*/
-							// filter bar
-							SNew(SDeviceBrowserFilterBar, Filter.ToSharedRef())
-						//]
-				]
+			[
+				/*
+				SNew(SExpandableArea)
+					.AreaTitle(LOCTEXT("FilterBarAreaTitle", "Device Filter").ToString())
+					.InitiallyCollapsed(true)
+					.Padding(FMargin(8.0f, 8.0f, 8.0f, 4.0f))
+					.BodyContent()
+					[*/
+						// filter bar
+						SNew(SDeviceBrowserFilterBar, Filter.ToSharedRef())
+					//]
+			]
 
 		+ SVerticalBox::Slot()
 			.FillHeight(1.0f)
@@ -58,16 +91,16 @@ void SDeviceBrowser::Construct( const FArguments& InArgs, const FDeviceManagerMo
 						SAssignNew(DeviceServiceListView, SListView<ITargetDeviceServicePtr>)
 							.ItemHeight(20.0f)
 							.ListItemsSource(&DeviceServiceList)
-							.OnContextMenuOpening(this, &SDeviceBrowser::HandleDeviceServiceListViewContextMenuOpening)
-							.OnGenerateRow(this, &SDeviceBrowser::HandleDeviceServiceListViewGenerateRow)
-							.OnSelectionChanged(this, &SDeviceBrowser::HandleDeviceServiceListViewSelectionChanged)
+							.OnContextMenuOpening_Lambda(DeviceServiceListViewContextMenuOpening)
+							.OnGenerateRow_Lambda(DeviceServiceListViewGenerateRow)
+							.OnSelectionChanged_Lambda(DeviceServiceListViewSelectionChanged)
 							.SelectionMode(ESelectionMode::Single)
 							.HeaderRow
 							(
 								SNew(SHeaderRow)
 
 								+ SHeaderRow::Column("Icon")
-									.DefaultLabel( FText::FromString(TEXT(" ")) )
+									.DefaultLabel(FText::FromString(TEXT(" ")))
 									.FixedWidth(32.0f)
 
 								+ SHeaderRow::Column("Name")
@@ -108,14 +141,15 @@ void SDeviceBrowser::Construct( const FArguments& InArgs, const FDeviceManagerMo
 			]
 	];
 
-	DeviceServiceManager->OnServiceAdded().AddSP(this, &SDeviceBrowser::HandleDeviceServiceManagerServiceAdded);
-	DeviceServiceManager->OnServiceRemoved().AddSP(this, &SDeviceBrowser::HandleDeviceServiceManagerServiceRemoved);
-
-	Filter->OnFilterChanged().AddSP(this, &SDeviceBrowser::HandleFilterChanged);
+	DeviceServiceManager->OnServiceAdded().AddLambda([this](const TSharedRef<ITargetDeviceService, ESPMode::ThreadSafe>& AddedService) { NeedsServiceListRefresh = true; });
+	DeviceServiceManager->OnServiceRemoved().AddLambda([this](const TSharedRef<ITargetDeviceService, ESPMode::ThreadSafe>& RemovedService) { NeedsServiceListRefresh = true; });
+	
+	Filter->OnFilterChanged().AddLambda([this]() { ReloadDeviceServiceList(false); });
 }
 END_SLATE_FUNCTION_BUILD_OPTIMIZATION
 
-void SDeviceBrowser::Tick( const FGeometry& AllottedGeometry, const double InCurrentTime, const float InDeltaTime )
+
+void SDeviceBrowser::Tick(const FGeometry& AllottedGeometry, const double InCurrentTime, const float InDeltaTime)
 {
 	//@TODO Passive - Only happens in response to the addition or removal of a device to the device service manager
 	if (NeedsServiceListRefresh)
@@ -129,7 +163,7 @@ void SDeviceBrowser::Tick( const FGeometry& AllottedGeometry, const double InCur
 /* SDeviceBrowser callbacks
  *****************************************************************************/
 
-void SDeviceBrowser::ReloadDeviceServiceList( bool FullyReload )
+void SDeviceBrowser::ReloadDeviceServiceList(bool FullyReload)
 {
 	// reload target device service list
 	if (FullyReload)
@@ -145,7 +179,7 @@ void SDeviceBrowser::ReloadDeviceServiceList( bool FullyReload )
 
 	for (int32 DeviceServiceIndex = 0; DeviceServiceIndex < AvailableDeviceServices.Num(); ++DeviceServiceIndex)
 	{
-		const ITargetDeviceServicePtr& DeviceService = AvailableDeviceServices[DeviceServiceIndex];
+		const TSharedPtr<ITargetDeviceService, ESPMode::ThreadSafe>& DeviceService = AvailableDeviceServices[DeviceServiceIndex];
 
 		if (Filter->FilterDeviceService(DeviceService.ToSharedRef()))
 		{
@@ -158,60 +192,4 @@ void SDeviceBrowser::ReloadDeviceServiceList( bool FullyReload )
 }
 
 
-/* SDeviceBrowser callbacks
- *****************************************************************************/
-
-
-TSharedPtr<SWidget> SDeviceBrowser::HandleDeviceServiceListViewContextMenuOpening( )
-{
-	TArray<ITargetDeviceServicePtr> SelectedDeviceServices = DeviceServiceListView->GetSelectedItems();
-
-	if (SelectedDeviceServices.Num() > 0)
-	{
-		return SNew(SDeviceBrowserContextMenu, UICommandList);
-	}
-	
-	return NULL;
-}
-
-
-TSharedRef<ITableRow> SDeviceBrowser::HandleDeviceServiceListViewGenerateRow( ITargetDeviceServicePtr DeviceService, const TSharedRef<STableViewBase>& OwnerTable )
-{
-	return SNew(SDeviceBrowserDeviceListRow, OwnerTable)
-		.DeviceService(DeviceService)
-		.HighlightText(this, &SDeviceBrowser::HandleDeviceServiceListViewHighlightText)
-		.ToolTip(SNew(SDeviceBrowserTooltip, DeviceService.ToSharedRef()));
-}
-
-
-void SDeviceBrowser::HandleDeviceServiceListViewSelectionChanged( ITargetDeviceServicePtr Selection, ESelectInfo::Type SelectInfo )
-{
-	Model->SelectDeviceService(Selection);
-}
-
-
-FText SDeviceBrowser::HandleDeviceServiceListViewHighlightText( ) const
-{
-	return Filter->GetDeviceSearchText();
-}
-
-
-void SDeviceBrowser::HandleDeviceServiceManagerServiceAdded( const ITargetDeviceServiceRef& AddedService )
-{
-	NeedsServiceListRefresh = true;
-}
-
-
-void SDeviceBrowser::HandleDeviceServiceManagerServiceRemoved( const ITargetDeviceServiceRef& RemovedService )
-{
-	NeedsServiceListRefresh = true;
-}
-
-
-void SDeviceBrowser::HandleFilterChanged( )
-{
-	ReloadDeviceServiceList(false);
-}
-
-
 #undef LOCTEXT_NAMESPACE
diff --git a/Engine/Source/Developer/DeviceManager/Private/Widgets/Browser/SDeviceBrowser.h b/Engine/Source/Developer/DeviceManager/Private/Widgets/Browser/SDeviceBrowser.h
--- a/Engine/Source/Developer/DeviceManager/Private/Widgets/Browser/SDeviceBrowser.h
+++ b/Engine/Source/Developer/DeviceManager/Private/Widgets/Browser/SDeviceBrowser.h
@@ -2,21 +2,24 @@
 
 #pragma once
 
-#include "CoreMinimal.h"
+#include "CoreTypes.h"
+#include "Containers/Array.h"
+#include "ITargetDeviceService.h"
+#include "Templates/SharedPointer.h"
 #include "Widgets/DeclarativeSyntaxSupport.h"
 #include "Widgets/SCompoundWidget.h"
-#include "Widgets/Views/STableViewBase.h"
-#include "Widgets/Views/STableRow.h"
-#include "Interfaces/ITargetDeviceServiceManager.h"
-#include "Models/DeviceBrowserFilter.h"
-#include "Models/DeviceManagerModel.h"
+#include "Widgets/Views/SListView.h"
 
+class FDeviceBrowserFilter;
+class FDeviceManagerModel;
 class FUICommandList;
+class ITargetDeviceServiceManager;
+
 
 /**
  * Delegate type for selection changes in the device browser.
  *
- * The first parameter is the newly selected device service (or NULL if none was selected).
+ * The first parameter is the newly selected device service (or nullptr if none was selected).
  */
 DECLARE_DELEGATE_OneParam(FOnDeviceBrowserSelectionChanged, const ITargetDeviceServicePtr&)
 
@@ -42,13 +45,13 @@ public:
 	 * @param InDeviceServiceManager The target device service manager to use.
 	 * @param InUICommandList The UI command list to use.
 	 */
-	void Construct( const FArguments& InArgs, const FDeviceManagerModelRef& InModel, const ITargetDeviceServiceManagerRef& InDeviceServiceManager, const TSharedPtr<FUICommandList>& InUICommandList );
+	void Construct(const FArguments& InArgs, const TSharedRef<FDeviceManagerModel>& InModel, const TSharedRef<ITargetDeviceServiceManager>& InDeviceServiceManager, const TSharedPtr<FUICommandList>& InUICommandList);
 
 public:
 
-	// SCompoundWidget overrides
+	//~ SCompoundWidget overrides
 
-	virtual void Tick( const FGeometry& AllottedGeometry, const double InCurrentTime, const float InDeltaTime ) override;
+	virtual void Tick(const FGeometry& AllottedGeometry, const double InCurrentTime, const float InDeltaTime) override;
 
 protected:
 
@@ -57,54 +60,31 @@ protected:
 	 *
 	 * @param FullyReload Whether to fully reload the service entries, or only re-apply filtering.
 	 */
-	void ReloadDeviceServiceList( bool FullyReload );
-
-private:
-
-	// Callback for generating a context menu in the session list view.
-	TSharedPtr<SWidget> HandleDeviceServiceListViewContextMenuOpening( );
-
-	// Callback for generating a row widget in the device service list view.
-	TSharedRef<ITableRow> HandleDeviceServiceListViewGenerateRow( ITargetDeviceServicePtr DeviceService, const TSharedRef<STableViewBase>& OwnerTable );
-
-	// Callback for getting the highlight text in session rows.
-	FText HandleDeviceServiceListViewHighlightText( ) const;
-
-	// Callback for selecting items in the devices list view.
-	void HandleDeviceServiceListViewSelectionChanged( ITargetDeviceServicePtr Selection, ESelectInfo::Type SelectInfo );
-
-	// Callback for added target device services.
-	void HandleDeviceServiceManagerServiceAdded( const ITargetDeviceServiceRef& AddedService );
-
-	// Callback for added target device services.
-	void HandleDeviceServiceManagerServiceRemoved( const ITargetDeviceServiceRef& RemovedService );
-
-	// Callback for changing the filter settings.
-	void HandleFilterChanged( );
+	void ReloadDeviceServiceList(bool FullyReload);
 
 private:
 
-	// Holds the list of all available target device services.
-	TArray<ITargetDeviceServicePtr> AvailableDeviceServices;
+	/** Holds the list of all available target device services. */
+	TArray<TSharedPtr<ITargetDeviceService, ESPMode::ThreadSafe>> AvailableDeviceServices;
 
-	// Holds the list of filtered target device services for display.
-	TArray<ITargetDeviceServicePtr> DeviceServiceList;
+	/** Holds the list of filtered target device services for display. */
+	TArray<TSharedPtr<ITargetDeviceService, ESPMode::ThreadSafe>> DeviceServiceList;
 
-	// Holds the list view for filtered target device services.
-	TSharedPtr<SListView<ITargetDeviceServicePtr> > DeviceServiceListView;
+	/** Holds the list view for filtered target device services. */
+	TSharedPtr<SListView<ITargetDeviceServicePtr>> DeviceServiceListView;
 
-	// Holds a pointer to the target device service manager.
-	ITargetDeviceServiceManagerPtr DeviceServiceManager;
+	/** Holds a pointer to the target device service manager. */
+	TSharedPtr<ITargetDeviceServiceManager> DeviceServiceManager;
 
-	// Holds the filter model.
-	FDeviceBrowserFilterPtr Filter;
+	/** Holds the filter model. */
+	TSharedPtr<FDeviceBrowserFilter> Filter;
 
-	// Holds a pointer the device manager's view model.
-	FDeviceManagerModelPtr Model;
+	/** Holds a pointer the device manager's view model. */
+	TSharedPtr<FDeviceManagerModel> Model;
 
-	// Holds a flag indicating whether the list of target device services needs to be refreshed.
+	/** Holds a flag indicating whether the list of target device services needs to be refreshed. */
 	bool NeedsServiceListRefresh;
 
-	// Holds a pointer to the command list for controlling the device.
+	/** Holds a pointer to the command list for controlling the device. */
 	TSharedPtr<FUICommandList> UICommandList;
 };
diff --git a/Engine/Source/Developer/DeviceManager/Private/Widgets/Browser/SDeviceBrowserContextMenu.h b/Engine/Source/Developer/DeviceManager/Private/Widgets/Browser/SDeviceBrowserContextMenu.h
--- a/Engine/Source/Developer/DeviceManager/Private/Widgets/Browser/SDeviceBrowserContextMenu.h
+++ b/Engine/Source/Developer/DeviceManager/Private/Widgets/Browser/SDeviceBrowserContextMenu.h
@@ -2,17 +2,20 @@
 
 #pragma once
 
-#include "CoreMinimal.h"
+#include "CoreTypes.h"
+#include "EditorStyleSet.h"
+#include "Framework/MultiBox/MultiBoxBuilder.h"
 #include "Widgets/DeclarativeSyntaxSupport.h"
-#include "Widgets/SWidget.h"
 #include "Widgets/SCompoundWidget.h"
+#include "Widgets/SWidget.h"
 #include "Widgets/Layout/SBorder.h"
-#include "EditorStyleSet.h"
-#include "Framework/MultiBox/MultiBoxBuilder.h"
+
 #include "Models/DeviceDetailsCommands.h"
 
+
 #define LOCTEXT_NAMESPACE "SDeviceBrowserContextMenu"
 
+
 /**
  * Implements a context menu for the device browser list view.
  */
@@ -32,7 +35,7 @@ public:
 	 * @param InArgs The declaration data for this widget.
 	 * @param InUICommandList The UI command list to use.
 	 */
-	void Construct( const FArguments& InArgs, const TSharedPtr<FUICommandList>& InUICommandList )
+	void Construct(const FArguments& InArgs, const TSharedPtr<FUICommandList>& InUICommandList)
 	{
 		ChildSlot
 		[
@@ -48,11 +51,11 @@ public:
 protected:
 
 	/**
-	 * Builds the context menu widget.
+	 * Build the context menu widget.
 	 *
 	 * @return The context menu.
 	 */
-	TSharedRef<SWidget> MakeContextMenu( const TSharedPtr<FUICommandList>& CommandList )
+	TSharedRef<SWidget> MakeContextMenu(const TSharedPtr<FUICommandList>& CommandList)
 	{
 		FMenuBuilder MenuBuilder(true, CommandList);
 
diff --git a/Engine/Source/Developer/DeviceManager/Private/Widgets/Browser/SDeviceBrowserDeviceAdder.cpp b/Engine/Source/Developer/DeviceManager/Private/Widgets/Browser/SDeviceBrowserDeviceAdder.cpp
--- a/Engine/Source/Developer/DeviceManager/Private/Widgets/Browser/SDeviceBrowserDeviceAdder.cpp
+++ b/Engine/Source/Developer/DeviceManager/Private/Widgets/Browser/SDeviceBrowserDeviceAdder.cpp
@@ -1,12 +1,15 @@
 // Copyright 1998-2017 Epic Games, Inc. All Rights Reserved.
 
-#include "Widgets/Browser/SDeviceBrowserDeviceAdder.h"
+#include "SDeviceBrowserDeviceAdder.h"
+
+#include "Internationalization/Text.h"
+#include "ITargetDeviceServiceManager.h"
 #include "Misc/MessageDialog.h"
 #include "Misc/CoreMisc.h"
 #include "SlateOptMacros.h"
 #include "Widgets/Images/SImage.h"
-#include "Widgets/Input/SEditableTextBox.h"
 #include "Widgets/Input/SButton.h"
+#include "Widgets/Input/SEditableTextBox.h"
 #include "EditorStyleSet.h"
 #include "Interfaces/ITargetPlatform.h"
 #include "Interfaces/ITargetPlatformManagerModule.h"
@@ -20,10 +23,148 @@
  *****************************************************************************/
 
 BEGIN_SLATE_FUNCTION_BUILD_OPTIMIZATION
-void SDeviceBrowserDeviceAdder::Construct( const FArguments& InArgs, const ITargetDeviceServiceManagerRef& InDeviceServiceManager )
+void SDeviceBrowserDeviceAdder::Construct(const FArguments& InArgs, const TSharedRef<ITargetDeviceServiceManager>& InDeviceServiceManager)
 {
 	DeviceServiceManager = InDeviceServiceManager;
 
+	// callback for clicking of the Add button
+	auto AddButtonClicked = [this]() -> FReply {
+		ITargetPlatform* TargetPlatform = GetTargetPlatformManager()->FindTargetPlatform(*PlatformComboBox->GetSelectedItem());
+
+		FString DeviceIdString = DeviceIdTextBox->GetText().ToString();
+		bool bAdded = TargetPlatform->AddDevice(DeviceIdString, false);
+		if (bAdded)
+		{
+			// pass credentials to the newly added device
+			if (TargetPlatform->RequiresUserCredentials())
+			{
+				// We cannot guess the device id, so we have to look it up by name
+				TArray<ITargetDevicePtr> Devices;
+				TargetPlatform->GetAllDevices(Devices);
+				for (ITargetDevicePtr Device : Devices)
+				{
+					if (Device.IsValid() && Device->GetId().GetDeviceName() == DeviceIdString)
+					{
+						FString UserNameString = UserNameTextBox->GetText().ToString();
+						FString UserPassString = UserPasswordTextBox->GetText().ToString();
+
+						Device->SetUserCredentials(UserNameString, UserPassString);
+					}
+				}
+			}
+
+			DeviceIdTextBox->SetText(FText::GetEmpty());
+			DeviceNameTextBox->SetText(FText::GetEmpty());
+			UserNameTextBox->SetText(FText::GetEmpty());
+			UserPasswordTextBox->SetText(FText::GetEmpty());
+		}
+		else
+		{
+			FMessageDialog::Open(EAppMsgType::Ok, LOCTEXT("DeviceAdderFailedToAddDeviceMessage", "Failed to add the device!"));
+		}
+
+		return FReply::Handled();
+	};
+
+	// callback for determining the enabled state of the 'Add' button
+	auto AddButtonIsEnabled = [this]() -> bool {
+		TSharedPtr<FString> PlatformName = PlatformComboBox->GetSelectedItem();
+
+		if (PlatformName.IsValid())
+		{
+			FString TextCheck = DeviceNameTextBox->GetText().ToString();
+			TextCheck.Trim();
+
+			if (!TextCheck.IsEmpty())
+			{
+				ITargetPlatform* Platform = GetTargetPlatformManager()->FindTargetPlatform(*PlatformName);
+
+				if (Platform != nullptr)
+				{
+					if (!Platform->RequiresUserCredentials())
+					{
+						return true;
+					}
+
+					// check user/password as well
+					TextCheck = UserNameTextBox->GetText().ToString();
+					TextCheck.Trim();
+
+					if (!TextCheck.IsEmpty())
+					{
+						// do not trim the password
+						return !UserPasswordTextBox->GetText().ToString().IsEmpty();
+					}
+				}
+			}
+		}
+
+		return false;
+	};
+
+	// callback for determining the visibility of the credentials box
+	auto CredentialsBoxVisibility = [this]() -> EVisibility {
+		TSharedPtr<FString> PlatformName = PlatformComboBox->GetSelectedItem();
+
+		if (PlatformName.IsValid())
+		{
+			ITargetPlatform* Platform = GetTargetPlatformManager()->FindTargetPlatform(*PlatformName);
+
+			if ((Platform != nullptr) && Platform->RequiresUserCredentials())
+			{
+				return EVisibility::Visible;
+			}
+		}
+
+		return EVisibility::Collapsed;
+	};
+
+	// callback for changes in the device name text box
+	auto DeviceNameTextBoxTextChanged = [this](const FString& Text) {
+		DetermineAddUnlistedButtonVisibility();
+	};
+
+	// callback for getting the name of the selected platform
+	auto PlatformComboBoxContentText = [this]() -> FText {
+		TSharedPtr<FString> SelectedPlatform = PlatformComboBox->GetSelectedItem();
+		return SelectedPlatform.IsValid() ? FText::FromString(*SelectedPlatform) : LOCTEXT("SelectAPlatform", "Select a Platform");
+	};
+
+	// callback for generating widgets for the platforms combo box
+	auto PlatformComboBoxGenerateWidget = [this](TSharedPtr<FString> Item) -> TSharedRef<SWidget> {
+		const PlatformInfo::FPlatformInfo* const PlatformInfo = PlatformInfo::FindPlatformInfo(**Item);
+
+		return
+			SNew(SHorizontalBox)
+
+			+ SHorizontalBox::Slot()
+				.AutoWidth()
+				.HAlign(HAlign_Left)
+				[
+					SNew(SBox)
+						.WidthOverride(24)
+						.HeightOverride(24)
+						[
+							SNew(SImage)
+								.Image((PlatformInfo) ? FEditorStyle::GetBrush(PlatformInfo->GetIconStyleName(PlatformInfo::EPlatformIconSize::Normal)) : FStyleDefaults::GetNoBrush())
+						]
+				]
+
+			+ SHorizontalBox::Slot()
+				.Padding(4.0f, 0.0f, 0.0f, 0.0f)
+				.VAlign(VAlign_Center)
+				[
+					SNew(STextBlock)
+						.Text(FText::FromString(*Item))
+				];
+	};
+
+	// callback for handling platform selection changes
+	auto PlatformComboBoxSelectionChanged = [this](TSharedPtr<FString> StringItem, ESelectInfo::Type SelectInfo) {
+		// @todo
+	};
+
+	// construct children
 	ChildSlot
 	[
 		SNew(SVerticalBox)
@@ -53,14 +194,14 @@ void SDeviceBrowserDeviceAdder::Construct( const FArguments& InArgs, const ITarg
 							.HAlign(HAlign_Left)
 							.Padding(0.0f, 4.0f, 0.0f, 0.0f)
 							[
-								SAssignNew(PlatformComboBox, SComboBox<TSharedPtr<FString> >)
+								SAssignNew(PlatformComboBox, SComboBox<TSharedPtr<FString>>)
 									.ContentPadding(FMargin(6.0f, 2.0f))
 									.OptionsSource(&PlatformList)
-									.OnGenerateWidget(this, &SDeviceBrowserDeviceAdder::HandlePlatformComboBoxGenerateWidget)
-									.OnSelectionChanged(this, &SDeviceBrowserDeviceAdder::HandlePlatformComboBoxSelectionChanged)
+									.OnGenerateWidget_Lambda(PlatformComboBoxGenerateWidget)
+									.OnSelectionChanged_Lambda(PlatformComboBoxSelectionChanged)
 									[
 										SNew(STextBlock)
-											.Text(this, &SDeviceBrowserDeviceAdder::HandlePlatformComboBoxContentText)
+											.Text_Lambda(PlatformComboBoxContentText)
 									]
 							]
 					]
@@ -121,9 +262,9 @@ void SDeviceBrowserDeviceAdder::Construct( const FArguments& InArgs, const ITarg
 					[
 						SAssignNew(AddButton, SButton)
 							.ContentPadding(FMargin(9.0, 2.0))
-							.IsEnabled(this, &SDeviceBrowserDeviceAdder::HandleAddButtonIsEnabled)
+							.IsEnabled_Lambda(AddButtonIsEnabled)
 							.Text(LOCTEXT("AddButtonText", "Add"))
-							.OnClicked(this, &SDeviceBrowserDeviceAdder::HandleAddButtonClicked)
+							.OnClicked_Lambda(AddButtonClicked)
 					]
 			]
 
@@ -131,7 +272,7 @@ void SDeviceBrowserDeviceAdder::Construct( const FArguments& InArgs, const ITarg
 			.AutoHeight()
 			[
 				SNew(SHorizontalBox)
-					.Visibility(this, &SDeviceBrowserDeviceAdder::HandleCredentialsBoxVisibility)
+					.Visibility_Lambda(CredentialsBoxVisibility)
 
 				// user name input
 				+ SHorizontalBox::Slot()
@@ -190,7 +331,7 @@ END_SLATE_FUNCTION_BUILD_OPTIMIZATION
 /* SDeviceBrowserDeviceAdder implementation
  *****************************************************************************/
 
-void SDeviceBrowserDeviceAdder::DetermineAddUnlistedButtonVisibility( )
+void SDeviceBrowserDeviceAdder::DetermineAddUnlistedButtonVisibility()
 {
 	if (PlatformComboBox->GetSelectedItem().IsValid())
 	{
@@ -202,7 +343,7 @@ void SDeviceBrowserDeviceAdder::DetermineAddUnlistedButtonVisibility( )
 }
 
 
-void SDeviceBrowserDeviceAdder::RefreshPlatformList( )
+void SDeviceBrowserDeviceAdder::RefreshPlatformList()
 {
 	TArray<ITargetPlatform*> Platforms = GetTargetPlatformManager()->GetTargetPlatforms();
 
@@ -217,152 +358,4 @@ void SDeviceBrowserDeviceAdder::RefreshPlatformList( )
 }
 
 
-/* SDeviceBrowserDeviceAdder event handlers
- *****************************************************************************/
-
-FReply SDeviceBrowserDeviceAdder::HandleAddButtonClicked( )
-{
-	ITargetPlatform* TargetPlatform = GetTargetPlatformManager()->FindTargetPlatform(*PlatformComboBox->GetSelectedItem());
-
-	FString DeviceIdString = DeviceIdTextBox->GetText().ToString();
-	bool bAdded = TargetPlatform->AddDevice(DeviceIdString, false);
-	if (bAdded)
-	{
-		// pass credentials to the newly added device
-		if (TargetPlatform->RequiresUserCredentials())
-		{
-			// We cannot guess the device id, so we have to look it up by name
-			TArray<ITargetDevicePtr> Devices;
-			TargetPlatform->GetAllDevices(Devices);
-			for (ITargetDevicePtr Device : Devices)
-			{
-				if (Device.IsValid() && Device->GetId().GetDeviceName() == DeviceIdString)
-				{
-					FString UserNameString = UserNameTextBox->GetText().ToString();
-					FString UserPassString = UserPasswordTextBox->GetText().ToString();
-
-					Device->SetUserCredentials(UserNameString, UserPassString);
-				}
-			}
-		}
-
-		DeviceIdTextBox->SetText(FText::GetEmpty());
-		DeviceNameTextBox->SetText(FText::GetEmpty());
-		UserNameTextBox->SetText(FText::GetEmpty());
-		UserPasswordTextBox->SetText(FText::GetEmpty());
-	}
-	else
-	{
-		FMessageDialog::Open(EAppMsgType::Ok, LOCTEXT("DeviceAdderFailedToAddDeviceMessage", "Failed to add the device!"));
-	}
-
-	return FReply::Handled();
-}
-
-
-bool SDeviceBrowserDeviceAdder::HandleAddButtonIsEnabled( ) const
-{
-	TSharedPtr<FString> PlatformName = PlatformComboBox->GetSelectedItem();
-
-	if (PlatformName.IsValid())
-	{
-		FString TextCheck = DeviceNameTextBox->GetText().ToString();
-		TextCheck.Trim();
-
-		if (!TextCheck.IsEmpty())
-		{
-			ITargetPlatform* Platform = GetTargetPlatformManager()->FindTargetPlatform(*PlatformName);
-
-			if (Platform != nullptr)
-			{
-				if (!Platform->RequiresUserCredentials())
-				{
-					return true;
-				}
-
-				// check user/password as well
-				TextCheck = UserNameTextBox->GetText().ToString();
-				TextCheck.Trim();
-
-				if (!TextCheck.IsEmpty())
-				{
-					// do not trim the password
-					return !UserPasswordTextBox->GetText().ToString().IsEmpty();
-				}
-			}
-		}
-	}
-
-	return false;
-}
-
-
-EVisibility SDeviceBrowserDeviceAdder::HandleCredentialsBoxVisibility( ) const
-{
-	TSharedPtr<FString> PlatformName = PlatformComboBox->GetSelectedItem();
-
-	if (PlatformName.IsValid())
-	{
-		ITargetPlatform* Platform = GetTargetPlatformManager()->FindTargetPlatform(*PlatformName);
-
-		if ((Platform != nullptr) && Platform->RequiresUserCredentials())
-		{
-			return EVisibility::Visible;
-		}
-	}
-
-	return EVisibility::Collapsed;
-}
-
-
-void SDeviceBrowserDeviceAdder::HandleDeviceNameTextBoxTextChanged (const FString& Text)
-{
-	DetermineAddUnlistedButtonVisibility();
-}
-
-
-FText SDeviceBrowserDeviceAdder::HandlePlatformComboBoxContentText( ) const
-{
-	TSharedPtr<FString> SelectedPlatform = PlatformComboBox->GetSelectedItem();
-
-	return SelectedPlatform.IsValid() ? FText::FromString(*SelectedPlatform) : LOCTEXT("SelectAPlatform", "Select a Platform");
-}
-
-
-TSharedRef<SWidget> SDeviceBrowserDeviceAdder::HandlePlatformComboBoxGenerateWidget( TSharedPtr<FString> Item )
-{
-	const PlatformInfo::FPlatformInfo* const PlatformInfo = PlatformInfo::FindPlatformInfo(**Item);
-
-	return
-		SNew(SHorizontalBox)
-
-		+ SHorizontalBox::Slot()
-			.AutoWidth()
-			.HAlign(HAlign_Left)
-			[
-				SNew(SBox)
-					.WidthOverride(24)
-					.HeightOverride(24)
-					[
-						SNew(SImage)
-							.Image((PlatformInfo) ? FEditorStyle::GetBrush(PlatformInfo->GetIconStyleName(PlatformInfo::EPlatformIconSize::Normal)) : FStyleDefaults::GetNoBrush())
-					]
-			]
-
-		+ SHorizontalBox::Slot()
-			.Padding(4.0f, 0.0f, 0.0f, 0.0f)
-			.VAlign(VAlign_Center)
-			[
-				SNew(STextBlock)
-					.Text(FText::FromString(*Item))
-			];
-}
-
-
-void SDeviceBrowserDeviceAdder::HandlePlatformComboBoxSelectionChanged( TSharedPtr<FString> StringItem, ESelectInfo::Type SelectInfo )
-{
-
-}
-
-
 #undef LOCTEXT_NAMESPACE
diff --git a/Engine/Source/Developer/DeviceManager/Private/Widgets/Browser/SDeviceBrowserDeviceAdder.h b/Engine/Source/Developer/DeviceManager/Private/Widgets/Browser/SDeviceBrowserDeviceAdder.h
--- a/Engine/Source/Developer/DeviceManager/Private/Widgets/Browser/SDeviceBrowserDeviceAdder.h
+++ b/Engine/Source/Developer/DeviceManager/Private/Widgets/Browser/SDeviceBrowserDeviceAdder.h
@@ -2,17 +2,19 @@
 
 #pragma once
 
-#include "CoreMinimal.h"
-#include "Layout/Visibility.h"
+#include "CoreTypes.h"
+#include "Containers/Array.h"
+#include "Templates/SharedPointer.h"
 #include "Widgets/DeclarativeSyntaxSupport.h"
-#include "Input/Reply.h"
-#include "Widgets/SWidget.h"
 #include "Widgets/SCompoundWidget.h"
-#include "Interfaces/ITargetDeviceServiceManager.h"
 #include "Widgets/Input/SComboBox.h"
 
+class FString;
+class ITargetDeviceServiceManager;
 class SButton;
 class SEditableTextBox;
+class SOverlay;
+
 
 /**
  * Implements a widget for manually locating target devices.
@@ -28,74 +30,47 @@ public:
 public:
 
 	/**
-	 * Constructs the widget.
+	 * Construct the widget.
 	 *
 	 * @param InArgs The construction arguments.
 	 * @param InDeviceManager The target device manager to use.
 	 */
-	void Construct( const FArguments& InArgs, const ITargetDeviceServiceManagerRef& InDeviceServiceManager );
+	void Construct(const FArguments& InArgs, const TSharedRef<ITargetDeviceServiceManager>& InDeviceServiceManager);
 
 protected:
 
-	/**
-	 * Determines the visibility of the AddUnlistedButton.
-	 */
-	void DetermineAddUnlistedButtonVisibility( );
-
-	/**
-	 * Refreshes the list of known platforms.
-	 */
-	void RefreshPlatformList( );
-
-private:
-
-	// Callback for clicking of the Add button.
-	FReply HandleAddButtonClicked( );
-
-	// Callback for determining the enabled state of the 'Add' button.
-	bool HandleAddButtonIsEnabled( ) const;
-
-	// Callback for determining the visibility of the credentials box.
-	EVisibility HandleCredentialsBoxVisibility( ) const;
-
-	// Callback for changes in the device name text box.
-	void HandleDeviceNameTextBoxTextChanged ( const FString& Text );
-
-	// Callback for getting the name of the selected platform.
-	FText HandlePlatformComboBoxContentText( ) const;
-
-	// Callback for generating widgets for the platforms combo box.
-	TSharedRef<SWidget> HandlePlatformComboBoxGenerateWidget( TSharedPtr<FString> StringItem );
+	/** Determine the visibility of the AddUnlistedButton. */
+	void DetermineAddUnlistedButtonVisibility();
 
-	// Callback for handling platform selection changes.
-	void HandlePlatformComboBoxSelectionChanged( TSharedPtr<FString> StringItem, ESelectInfo::Type SelectInfo );
+	/** Refreshes the list of known platforms. */
+	void RefreshPlatformList();
 
 private:
 
-	// Holds the button for adding an unlisted device.
+	/** The button for adding an unlisted device. */
 	TSharedPtr<SButton> AddButton;
 
-	// Holds the device identifier text box.
+	/** The device identifier text box.  */
 	TSharedPtr<SEditableTextBox> DeviceIdTextBox;
 
-	// Holds a pointer to the target device service manager.
-	ITargetDeviceServiceManagerPtr DeviceServiceManager;
+	/** Holds a pointer to the target device service manager. */
+	TSharedPtr<ITargetDeviceServiceManager> DeviceServiceManager;
 
-	// Holds the device name text box.
+	/** The device name text box. */
 	TSharedPtr<SEditableTextBox> DeviceNameTextBox;
 
-	// Holds the user name text box.
+	/** The user name text box. */
 	TSharedPtr<SEditableTextBox> UserNameTextBox;
 
-	// Holds the user password text box.
+	/** The user password text box. */
 	TSharedPtr<SEditableTextBox> UserPasswordTextBox;
 
-	// Holds the turnable overlay with user data
+	/** The turnable overlay with user data. */
 	TSharedPtr<SOverlay> UserDataOverlay;
 
-	// Holds the platforms combo box.
+	/** The platforms combo box. */
 	TSharedPtr<SComboBox<TSharedPtr<FString>>> PlatformComboBox;
 
-	// Holds the list of known platforms.
+	/** The list of known platforms. */
 	TArray<TSharedPtr<FString>> PlatformList;
 };
diff --git a/Engine/Source/Developer/DeviceManager/Private/Widgets/Browser/SDeviceBrowserDeviceListRow.h b/Engine/Source/Developer/DeviceManager/Private/Widgets/Browser/SDeviceBrowserDeviceListRow.h
--- a/Engine/Source/Developer/DeviceManager/Private/Widgets/Browser/SDeviceBrowserDeviceListRow.h
+++ b/Engine/Source/Developer/DeviceManager/Private/Widgets/Browser/SDeviceBrowserDeviceListRow.h
@@ -2,36 +2,40 @@
 
 #pragma once
 
-#include "CoreMinimal.h"
+#include "CoreTypes.h"
+#include "EditorStyleSet.h"
+#include "HAL/PlatformProcess.h"
+#include "Internationalization/Text.h"
+#include "ITargetDeviceProxy.h"
+#include "ITargetDeviceService.h"
+#include "Layout/Margin.h"
 #include "Misc/Attribute.h"
-#include "Widgets/SNullWidget.h"
-#include "Widgets/DeclarativeSyntaxSupport.h"
+#include "PlatformInfo.h"
+#include "SlateOptMacros.h"
 #include "Styling/SlateColor.h"
-#include "Widgets/SWidget.h"
-#include "Layout/Margin.h"
-#include "Widgets/Views/STableViewBase.h"
 #include "Styling/SlateTypes.h"
 #include "Styling/StyleDefaults.h"
+#include "Templates/SharedPointer.h"
+#include "Widgets/DeclarativeSyntaxSupport.h"
+#include "Widgets/SNullWidget.h"
+#include "Widgets/SWidget.h"
+#include "Widgets/Images/SImage.h"
+#include "Widgets/Input/SCheckBox.h"
 #include "Widgets/Layout/SBox.h"
-#include "Widgets/Views/STableRow.h"
-#include "Interfaces/ITargetDeviceService.h"
-#include "SlateOptMacros.h"
-#include "EditorStyleSet.h"
 #include "Widgets/Text/STextBlock.h"
 #include "Widgets/Views/SListView.h"
-#include "Widgets/Images/SImage.h"
-#include "PlatformInfo.h"
-#include "HAL/PlatformProcess.h"
-#include "Widgets/Input/SCheckBox.h"
-#include "Interfaces/ITargetDeviceProxy.h"
+#include "Widgets/Views/STableRow.h"
+#include "Widgets/Views/STableViewBase.h"
+
 
 #define LOCTEXT_NAMESPACE "SDeviceBrowserDeviceListRow"
 
+
 /**
  * Implements a row widget for the device list view.
  */
 class SDeviceBrowserDeviceListRow
-	: public SMultiColumnTableRow<ITargetDeviceProxyPtr>
+	: public SMultiColumnTableRow<TSharedPtr<ITargetDeviceProxy>>
 {
 public:
 
@@ -43,11 +47,11 @@ public:
 public:
 
 	/**
-	 * Constructs the widget.
+	 * Construct the widget.
 	 *
 	 * @param InArgs The arguments.
 	 */
-	void Construct( const FArguments& InArgs, const TSharedRef<STableViewBase>& InOwnerTableView )
+	void Construct(const FArguments& InArgs, const TSharedRef<STableViewBase>& InOwnerTableView)
 	{
 		check(InArgs._DeviceService.IsValid());
 
@@ -55,19 +59,15 @@ public:
 		HighlightText = InArgs._HighlightText;
 
 		
-		SMultiColumnTableRow<ITargetDeviceProxyPtr>::Construct(FSuperRowType::FArguments(), InOwnerTableView);
+		SMultiColumnTableRow<TSharedPtr<ITargetDeviceProxy>>::Construct(FSuperRowType::FArguments(), InOwnerTableView);
 	}
 
 public:
 
-	/**
-	 * Generates the widget for the specified column.
-	 *
-	 * @param ColumnName The name of the column to generate the widget for.
-	 * @return The widget.
-	 */
+	//~ SMultiColumnTableRow interface
+
 	BEGIN_SLATE_FUNCTION_BUILD_OPTIMIZATION
-	virtual TSharedRef<SWidget> GenerateWidgetForColumn( const FName& ColumnName ) override
+	virtual TSharedRef<SWidget> GenerateWidgetForColumn(const FName& ColumnName) override
 	{
 		if (ColumnName == TEXT("Claimed"))
 		{
@@ -148,7 +148,7 @@ public:
 private:
 
 	// Callback for getting the text in the 'Claimed' column.
-	FText HandleClaimedText( ) const
+	FText HandleClaimedText() const
 	{
 		FString ClaimUser = DeviceService->GetClaimUser();
 
@@ -161,19 +161,19 @@ private:
 	}
 
 	// Callback for getting the text in the 'Name' column.
-	FText HandleNameText( ) const
+	FText HandleNameText() const
 	{
 		return FText::FromString(DeviceService->GetDeviceName());
 	}
 
 	// Callback for changing this row's Share check box state.
-	void HandleShareCheckBoxStateChanged( ECheckBoxState NewState )
+	void HandleShareCheckBoxStateChanged(ECheckBoxState NewState)
 	{
 		DeviceService->SetShared(NewState == ECheckBoxState::Checked);
 	}
 
 	// Callback for getting the state of the 'Share' check box.
-	ECheckBoxState HandleShareCheckBoxIsChecked( ) const
+	ECheckBoxState HandleShareCheckBoxIsChecked() const
 	{
 		if (DeviceService->IsShared())
 		{
@@ -184,13 +184,13 @@ private:
 	}
 
 	// Callback for getting the enabled state of the 'Share' check box.
-	bool HandleShareCheckBoxIsEnabled( ) const
+	bool HandleShareCheckBoxIsEnabled() const
 	{
 		return DeviceService->IsRunning();
 	}
 
 	// Callback for getting the status text.
-	FText HandleStatusTextBlockText( ) const
+	FText HandleStatusTextBlockText() const
 	{
 		ITargetDevicePtr TargetDevice = DeviceService->GetDevice();
 
@@ -208,7 +208,7 @@ private:
 	}
 
 	// Callback for getting the text color.
-	FSlateColor HandleTextColorAndOpacity( ) const
+	FSlateColor HandleTextColorAndOpacity() const
 	{
 		if (DeviceService->CanStart())
 		{
@@ -221,7 +221,7 @@ private:
 private:
 
 	// Holds the target device service used to populate this row.
-	ITargetDeviceServicePtr DeviceService;
+	TSharedPtr<ITargetDeviceService, ESPMode::ThreadSafe> DeviceService;
 
 	// Holds the highlight text for the log message.
 	TAttribute<FText> HighlightText;
diff --git a/Engine/Source/Developer/DeviceManager/Private/Widgets/Browser/SDeviceBrowserFilterBar.cpp b/Engine/Source/Developer/DeviceManager/Private/Widgets/Browser/SDeviceBrowserFilterBar.cpp
--- a/Engine/Source/Developer/DeviceManager/Private/Widgets/Browser/SDeviceBrowserFilterBar.cpp
+++ b/Engine/Source/Developer/DeviceManager/Private/Widgets/Browser/SDeviceBrowserFilterBar.cpp
@@ -1,15 +1,17 @@
 // Copyright 1998-2017 Epic Games, Inc. All Rights Reserved.
 
-#include "Widgets/Browser/SDeviceBrowserFilterBar.h"
+#include "SDeviceBrowserFilterBar.h"
+
+#include "EditorStyleSet.h"
+#include "PlatformInfo.h"
 #include "SlateOptMacros.h"
 #include "Widgets/Images/SImage.h"
-#include "Widgets/Text/STextBlock.h"
-#include "Widgets/Input/SComboButton.h"
-#include "Widgets/Views/SListView.h"
 #include "Widgets/Input/SCheckBox.h"
-#include "EditorStyleSet.h"
-#include "PlatformInfo.h"
+#include "Widgets/Input/SComboButton.h"
 #include "Widgets/Input/SSearchBox.h"
+#include "Widgets/Text/STextBlock.h"
+
+#include "Models/DeviceBrowserFilter.h"
 
 
 #define LOCTEXT_NAMESPACE "SDeviceBrowserFilterBar"
@@ -18,7 +20,7 @@
 /* SSessionBrowserFilterBar structors
  *****************************************************************************/
 
-SDeviceBrowserFilterBar::~SDeviceBrowserFilterBar( )
+SDeviceBrowserFilterBar::~SDeviceBrowserFilterBar()
 {
 	if (Filter.IsValid())
 	{
@@ -31,10 +33,75 @@ SDeviceBrowserFilterBar::~SDeviceBrowserFilterBar( )
  *****************************************************************************/
 
 BEGIN_SLATE_FUNCTION_BUILD_OPTIMIZATION
-void SDeviceBrowserFilterBar::Construct( const FArguments& InArgs, FDeviceBrowserFilterRef InFilter )
+void SDeviceBrowserFilterBar::Construct(const FArguments& InArgs, TSharedRef<FDeviceBrowserFilter> InFilter)
 {
 	Filter = InFilter;
 
+	// callback for filter model resets
+	auto FilterReset = [this]() {
+		FilterStringTextBox->SetText(Filter->GetDeviceSearchText());
+		PlatformListView->RequestListRefresh();
+	};
+
+	// callback for changing the filter string text box text
+	auto FilterStringTextChanged = [this](const FText& NewText) {
+		Filter->SetDeviceSearchString(NewText);
+	};
+
+	// callback for generating a row widget for the platform filter list
+	auto PlatformListViewGenerateRow = [this](TSharedPtr<FDeviceBrowserFilterEntry> PlatformEntry, const TSharedRef<STableViewBase>& OwnerTable) -> TSharedRef<ITableRow> {
+		const PlatformInfo::FPlatformInfo* const PlatformInfo = PlatformInfo::FindPlatformInfo(PlatformEntry->PlatformLookup);
+
+		return SNew(STableRow<TSharedPtr<FString> >, OwnerTable)
+			.Content()
+			[
+				SNew(SCheckBox)
+				.IsChecked_Lambda(
+					[=]() -> ECheckBoxState {
+						return Filter->IsPlatformEnabled(PlatformEntry->PlatformName)
+							? ECheckBoxState::Checked
+							: ECheckBoxState::Unchecked;
+					}
+				)
+				.Padding(FMargin(6.0, 2.0))
+				.OnCheckStateChanged_Lambda(
+					[=](ECheckBoxState CheckState) {
+						Filter->SetPlatformEnabled(PlatformEntry->PlatformName, CheckState == ECheckBoxState::Checked);
+					}
+				)
+				.Content()
+				[
+					SNew(SHorizontalBox)
+
+					+ SHorizontalBox::Slot()
+						.AutoWidth()
+						[
+							SNew(SBox)
+								.WidthOverride(24)
+								.HeightOverride(24)
+								[
+									SNew(SImage)
+									.Image((PlatformInfo) ? FEditorStyle::GetBrush(PlatformInfo->GetIconStyleName(PlatformInfo::EPlatformIconSize::Normal)) : FStyleDefaults::GetNoBrush())
+								]
+						]
+
+					+ SHorizontalBox::Slot()
+						.FillWidth(1.0f)
+						.Padding(4.0f, 0.0f, 0.0f, 0.0f)
+						.VAlign(VAlign_Center)
+						[
+							SNew(STextBlock)
+								.Text_Lambda(
+									[=]() -> FText {
+										return FText::Format(LOCTEXT("PlatformListRowFmt", "{0} ({1})"), FText::FromString(PlatformEntry->PlatformName), FText::AsNumber(Filter->GetServiceCountPerPlatform(PlatformEntry->PlatformName)));
+									}
+								)
+						]
+				]
+			];
+	};
+
+	// construct children
 	ChildSlot
 	[
 		SNew(SHorizontalBox)
@@ -60,7 +127,7 @@ void SDeviceBrowserFilterBar::Construct( const FArguments& InArgs, FDeviceBrowse
 						SAssignNew(PlatformListView, SListView<TSharedPtr<FDeviceBrowserFilterEntry> >)
 							.ItemHeight(24.0f)
 							.ListItemsSource(&Filter->GetFilteredPlatforms())
-							.OnGenerateRow(this, &SDeviceBrowserFilterBar::HandlePlatformListViewGenerateRow)
+							.OnGenerateRow_Lambda(PlatformListViewGenerateRow)
 					]
 			]
 
@@ -70,93 +137,13 @@ void SDeviceBrowserFilterBar::Construct( const FArguments& InArgs, FDeviceBrowse
 			[
 				// search box
 				SAssignNew(FilterStringTextBox, SSearchBox)
-				.HintText(LOCTEXT("SearchBoxHint", "Search devices"))
-				.OnTextChanged(this, &SDeviceBrowserFilterBar::HandleFilterStringTextChanged)
+					.HintText(LOCTEXT("SearchBoxHint", "Search devices"))
+					.OnTextChanged_Lambda(FilterStringTextChanged)
 			]
 
 	];
 
-	Filter->OnFilterReset().AddSP(this, &SDeviceBrowserFilterBar::HandleFilterReset);
-}
-END_SLATE_FUNCTION_BUILD_OPTIMIZATION
-
-
-/* SSessionBrowserFilterBar callbacks
- *****************************************************************************/
-
-void SDeviceBrowserFilterBar::HandleFilterReset( )
-{
-	FilterStringTextBox->SetText(Filter->GetDeviceSearchText());
-	PlatformListView->RequestListRefresh();
-}
-
-
-void SDeviceBrowserFilterBar::HandleFilterStringTextChanged( const FText& NewText )
-{
-	Filter->SetDeviceSearchString(NewText);
-}
-
-
-void SDeviceBrowserFilterBar::HandlePlatformListRowCheckStateChanged(ECheckBoxState CheckState, TSharedPtr<FDeviceBrowserFilterEntry> PlatformEntry)
-{
-	Filter->SetPlatformEnabled(PlatformEntry->PlatformName, CheckState == ECheckBoxState::Checked);
-}
-
-
-ECheckBoxState SDeviceBrowserFilterBar::HandlePlatformListRowIsChecked(TSharedPtr<FDeviceBrowserFilterEntry> PlatformEntry) const
-{
-	if (Filter->IsPlatformEnabled(PlatformEntry->PlatformName))
-	{
-		return ECheckBoxState::Checked;
-	}
-
-	return ECheckBoxState::Unchecked;
-}
-
-
-FText SDeviceBrowserFilterBar::HandlePlatformListRowText(TSharedPtr<FDeviceBrowserFilterEntry> PlatformEntry) const
-{
-	return FText::Format(LOCTEXT("PlatformListRowFmt", "{0} ({1})"), FText::FromString(PlatformEntry->PlatformName), FText::AsNumber(Filter->GetServiceCountPerPlatform(PlatformEntry->PlatformName)));
-}
-
-BEGIN_SLATE_FUNCTION_BUILD_OPTIMIZATION
-TSharedRef<ITableRow> SDeviceBrowserFilterBar::HandlePlatformListViewGenerateRow(TSharedPtr<FDeviceBrowserFilterEntry> PlatformEntry, const TSharedRef<STableViewBase>& OwnerTable)
-{
-	const PlatformInfo::FPlatformInfo* const PlatformInfo = PlatformInfo::FindPlatformInfo(PlatformEntry->PlatformLookup);
-
-	return SNew(STableRow<TSharedPtr<FString> >, OwnerTable)
-		.Content()
-		[
-			SNew(SCheckBox)
-			.IsChecked(this, &SDeviceBrowserFilterBar::HandlePlatformListRowIsChecked, PlatformEntry)
-			.Padding(FMargin(6.0, 2.0))
-			.OnCheckStateChanged(this, &SDeviceBrowserFilterBar::HandlePlatformListRowCheckStateChanged, PlatformEntry)
-			.Content()
-			[
-				SNew(SHorizontalBox)
-
-				+ SHorizontalBox::Slot()
-				.AutoWidth()
-				[
-					SNew(SBox)
-					.WidthOverride(24)
-					.HeightOverride(24)
-					[
-						SNew(SImage)
-						.Image((PlatformInfo) ? FEditorStyle::GetBrush(PlatformInfo->GetIconStyleName(PlatformInfo::EPlatformIconSize::Normal)) : FStyleDefaults::GetNoBrush())
-					]
-				]
-
-				+ SHorizontalBox::Slot()
-				.FillWidth(1.0f)
-				.Padding(4.0f, 0.0f, 0.0f, 0.0f)
-				.VAlign(VAlign_Center)
-				[
-					SNew(STextBlock)
-					.Text(this, &SDeviceBrowserFilterBar::HandlePlatformListRowText, PlatformEntry)
-				]
-			]
-		];
+	Filter->OnFilterReset().AddLambda(FilterReset);
 }
 END_SLATE_FUNCTION_BUILD_OPTIMIZATION
 
diff --git a/Engine/Source/Developer/DeviceManager/Private/Widgets/Browser/SDeviceBrowserFilterBar.h b/Engine/Source/Developer/DeviceManager/Private/Widgets/Browser/SDeviceBrowserFilterBar.h
--- a/Engine/Source/Developer/DeviceManager/Private/Widgets/Browser/SDeviceBrowserFilterBar.h
+++ b/Engine/Source/Developer/DeviceManager/Private/Widgets/Browser/SDeviceBrowserFilterBar.h
@@ -2,13 +2,17 @@
 
 #pragma once
 
-#include "CoreMinimal.h"
-#include "SlateFwd.h"
+#include "CoreTypes.h"
+#include "Templates/SharedPointer.h"
 #include "Widgets/DeclarativeSyntaxSupport.h"
 #include "Widgets/SCompoundWidget.h"
-#include "Widgets/Views/STableViewBase.h"
-#include "Widgets/Views/STableRow.h"
-#include "Models/DeviceBrowserFilter.h"
+#include "Widgets/Views/SListView.h"
+
+class FDeviceBrowserFilter;
+class SSearchBox;
+
+struct FDeviceBrowserFilterEntry;
+
 
 /**
  * Implements the device browser filter bar widget.
@@ -23,56 +27,32 @@ public:
 
 public:
 
-	/**
-	 * Default constructor.
-	 */
+	/** Default constructor. */
 	SDeviceBrowserFilterBar()
-		: Filter(NULL)
+		: Filter(nullptr)
 	{ }
 
-	/**
-	 * Destructor.
-	 */
+	/** Destructor. */
 	~SDeviceBrowserFilterBar();
 
 public:
 
 	/**
-	 * Construct this widget
+	 * Construct this widget.
 	 *
 	 * @param InArgs The declaration data for this widget.
 	 * @param InFilter The filter model to use.
 	 */
-	void Construct(const FArguments& InArgs, FDeviceBrowserFilterRef InFilter);
-
-private:
-
-	// Callback for filter model resets.
-	void HandleFilterReset();
-
-	// Callback for changing the filter string text box text.
-	void HandleFilterStringTextChanged(const FText& NewText);
-
-	// Callback for changing the checked state of the given platform filter row.
-	void HandlePlatformListRowCheckStateChanged(ECheckBoxState CheckState, TSharedPtr<FDeviceBrowserFilterEntry> PlatformEntry);
-
-	// Callback for getting the checked state of the given platform filter row.
-	ECheckBoxState HandlePlatformListRowIsChecked(TSharedPtr<FDeviceBrowserFilterEntry> PlatformEntry) const;
-
-	// Callback for getting the text for a row in the platform filter drop-down.
-	FText HandlePlatformListRowText(TSharedPtr<FDeviceBrowserFilterEntry> PlatformEntry) const;
-
-	// Generates a row widget for the platform filter list.
-	TSharedRef<ITableRow> HandlePlatformListViewGenerateRow(TSharedPtr<FDeviceBrowserFilterEntry> PlatformEntry, const TSharedRef<STableViewBase>& OwnerTable);
+	void Construct(const FArguments& InArgs, TSharedRef<FDeviceBrowserFilter> InFilter);
 
 private:
 
-	// Holds a pointer to the filter model.
-	FDeviceBrowserFilterPtr Filter;
+	/** Pointer to the filter model. */
+	TSharedPtr<FDeviceBrowserFilter> Filter;
 
-	// Holds the filter string text box.
+	/** The filter string text box. */
 	TSharedPtr<SSearchBox> FilterStringTextBox;
 
-	// Holds the platform filters list view.
+	/** The platform filters list view. */
 	TSharedPtr<SListView<TSharedPtr<FDeviceBrowserFilterEntry>>> PlatformListView;
 };
diff --git a/Engine/Source/Developer/DeviceManager/Private/Widgets/Browser/SDeviceBrowserTooltip.h b/Engine/Source/Developer/DeviceManager/Private/Widgets/Browser/SDeviceBrowserTooltip.h
--- a/Engine/Source/Developer/DeviceManager/Private/Widgets/Browser/SDeviceBrowserTooltip.h
+++ b/Engine/Source/Developer/DeviceManager/Private/Widgets/Browser/SDeviceBrowserTooltip.h
@@ -2,11 +2,15 @@
 
 #pragma once
 
-#include "CoreMinimal.h"
+#include "CoreTypes.h"
 #include "Widgets/DeclarativeSyntaxSupport.h"
 #include "Widgets/SToolTip.h"
+
 #include "Widgets/Shared/SDeviceQuickInfo.h"
 
+class ITargetDeviceService;
+
+
 #define LOCTEXT_NAMESPACE "SDeviceBrowserTooltip"
 
 
@@ -24,13 +28,13 @@ public:
 public:
 
 	/**
-	 * Constructs the widget.
+	 * Construct the widget.
 	 *
 	 * @param InArgs The construction arguments.
 	 * @param InDeviceServiceManager The target device service manager to use.
 	 * @param InDeviceManagerState The optional device manager view state.
 	 */
-	void Construct( const FArguments& InArgs, const ITargetDeviceServiceRef& InDeviceService )
+	void Construct(const FArguments& InArgs, const TSharedRef<ITargetDeviceService, ESPMode::ThreadSafe>& InDeviceService)
 	{
 		SToolTip::Construct(
 			SToolTip::FArguments()
diff --git a/Engine/Source/Developer/DeviceManager/Private/Widgets/Details/SDeviceDetails.cpp b/Engine/Source/Developer/DeviceManager/Private/Widgets/Details/SDeviceDetails.cpp
--- a/Engine/Source/Developer/DeviceManager/Private/Widgets/Details/SDeviceDetails.cpp
+++ b/Engine/Source/Developer/DeviceManager/Private/Widgets/Details/SDeviceDetails.cpp
@@ -7,7 +7,11 @@
 #include "EditorStyleSet.h"
 #include "Interfaces/ITargetPlatform.h"
 #include "Widgets/Shared/SDeviceQuickInfo.h"
-#include "Widgets/Views/SListView.h"
+#include "Widgets/Views/STableViewBase.h"
+#include "Widgets/Views/STableRow.h"
+
+#include "Models/DeviceDetailsFeature.h"
+#include "Models/DeviceManagerModel.h"
 #include "Widgets/Details/SDeviceDetailsFeatureListRow.h"
 
 
@@ -17,7 +21,7 @@
 /* SMessagingEndpoints structors
 *****************************************************************************/
 
-SDeviceDetails::~SDeviceDetails( )
+SDeviceDetails::~SDeviceDetails()
 {
 	if (Model.IsValid())
 	{
@@ -30,10 +34,31 @@ SDeviceDetails::~SDeviceDetails( )
  *****************************************************************************/
 
 BEGIN_SLATE_FUNCTION_BUILD_OPTIMIZATION
-void SDeviceDetails::Construct( const FArguments& InArgs, const FDeviceManagerModelRef& InModel )
+void SDeviceDetails::Construct(const FArguments& InArgs, const TSharedRef<FDeviceManagerModel>& InModel)
 {
 	Model = InModel;
 
+	// callback for getting the visibility of the details box.
+	auto DetailsBoxVisibility = [this]() -> EVisibility {
+		return Model->GetSelectedDeviceService().IsValid()
+			? EVisibility::Visible
+			: EVisibility::Hidden;
+	};
+
+	// callback for generating a row widget for the feature list view.
+	auto FeatureListGenerateRow = [](TSharedPtr<FDeviceDetailsFeature> Feature, const TSharedRef<STableViewBase>& OwnerTable) -> TSharedRef<ITableRow> {
+		return SNew(SDeviceDetailsFeatureListRow, OwnerTable, Feature.ToSharedRef());
+			//.Style(Style);
+	};
+
+	// callback for getting the visibility of the 'Select a device' message.
+	auto HandleSelectDeviceOverlayVisibility = [this]() -> EVisibility {
+		return Model->GetSelectedDeviceService().IsValid()
+			? EVisibility::Hidden
+			: EVisibility::Visible;
+	};
+
+	// construct children
 	ChildSlot
 	[
 		SNew(SOverlay)
@@ -41,7 +66,7 @@ void SDeviceDetails::Construct( const FArguments& InArgs, const FDeviceManagerMo
 		+ SOverlay::Slot()
 			[
 				SNew(SVerticalBox)
-					.Visibility(this, &SDeviceDetails::HandleDetailsBoxVisibility)
+					.Visibility_Lambda(DetailsBoxVisibility)
 
 				+ SVerticalBox::Slot()
 					.AutoHeight()
@@ -60,10 +85,10 @@ void SDeviceDetails::Construct( const FArguments& InArgs, const FDeviceManagerMo
 							.Padding(0.0f)
 							[
 								// feature list view
-								SAssignNew(FeatureListView, SListView<FDeviceDetailsFeaturePtr>)
+								SAssignNew(FeatureListView, SListView<TSharedPtr<FDeviceDetailsFeature>>)
 									.ItemHeight(24.0f)
 									.ListItemsSource(&FeatureList)
-									.OnGenerateRow(this, &SDeviceDetails::HandleFeatureListGenerateRow)
+									.OnGenerateRow_Lambda(FeatureListGenerateRow)
 									.SelectionMode(ESelectionMode::None)
 									.HeaderRow
 									(
@@ -88,7 +113,7 @@ void SDeviceDetails::Construct( const FArguments& InArgs, const FDeviceManagerMo
 				SNew(SBorder)
 					.BorderImage(FEditorStyle::GetBrush("NotificationList.ItemBackground"))
 					.Padding(8.0f)
-					.Visibility(this, &SDeviceDetails::HandleSelectDeviceOverlayVisibility)
+					.Visibility_Lambda(HandleSelectDeviceOverlayVisibility)
 					[
 						SNew(STextBlock)
 							.Text(LOCTEXT("SelectSessionOverlayText", "Please select a device from the Device Browser"))
@@ -96,84 +121,52 @@ void SDeviceDetails::Construct( const FArguments& InArgs, const FDeviceManagerMo
 			]
 	];
 
-	Model->OnSelectedDeviceServiceChanged().AddRaw(this, &SDeviceDetails::HandleModelSelectedDeviceServiceChanged);
-}
-END_SLATE_FUNCTION_BUILD_OPTIMIZATION
-
+	// callback for handling device service selection changes.
+	auto ModelSelectedDeviceServiceChanged = [this]() {
+		FeatureList.Empty();
 
-/* SDeviceDetails callbacks
- *****************************************************************************/
-
-EVisibility SDeviceDetails::HandleDetailsBoxVisibility( ) const
-{
-	if (Model->GetSelectedDeviceService().IsValid())
-	{
-		return EVisibility::Visible;
-	}
-
-	return EVisibility::Hidden;
-}
-
-
-TSharedRef<ITableRow> SDeviceDetails::HandleFeatureListGenerateRow( FDeviceDetailsFeaturePtr Feature, const TSharedRef<STableViewBase>& OwnerTable )
-{
-	return SNew(SDeviceDetailsFeatureListRow, OwnerTable, Feature.ToSharedRef());
-//		.Style(Style);
-}
+		ITargetDeviceServicePtr DeviceService = Model->GetSelectedDeviceService();
 
-
-void SDeviceDetails::HandleModelSelectedDeviceServiceChanged( )
-{
-	FeatureList.Empty();
-
-	ITargetDeviceServicePtr DeviceService = Model->GetSelectedDeviceService();
-
-	if (DeviceService.IsValid())
-	{
-		ITargetDevicePtr TargetDevice = DeviceService->GetDevice();
-
-		if (TargetDevice.IsValid())
+		if (DeviceService.IsValid())
 		{
-			const ITargetPlatform& TargetPlatform = TargetDevice->GetTargetPlatform();
-
-			// platform features
-			FeatureList.Add(MakeShareable(new FDeviceDetailsFeature(TEXT("AudioStreaming"), TargetPlatform.SupportsFeature(ETargetPlatformFeatures::AudioStreaming))));
-			FeatureList.Add(MakeShareable(new FDeviceDetailsFeature(TEXT("DistanceFieldShadows"), TargetPlatform.SupportsFeature(ETargetPlatformFeatures::DistanceFieldShadows))));
-			FeatureList.Add(MakeShareable(new FDeviceDetailsFeature(TEXT("GrayscaleSRGB"), TargetPlatform.SupportsFeature(ETargetPlatformFeatures::GrayscaleSRGB))));
-			FeatureList.Add(MakeShareable(new FDeviceDetailsFeature(TEXT("HighQualityLightmaps"), TargetPlatform.SupportsFeature(ETargetPlatformFeatures::HighQualityLightmaps))));
-			FeatureList.Add(MakeShareable(new FDeviceDetailsFeature(TEXT("LowQualityLightmaps"), TargetPlatform.SupportsFeature(ETargetPlatformFeatures::LowQualityLightmaps))));
-			FeatureList.Add(MakeShareable(new FDeviceDetailsFeature(TEXT("MultipleGameInstances"), TargetPlatform.SupportsFeature(ETargetPlatformFeatures::MultipleGameInstances))));
-			FeatureList.Add(MakeShareable(new FDeviceDetailsFeature(TEXT("Packaging"), TargetPlatform.SupportsFeature(ETargetPlatformFeatures::Packaging))));
-			FeatureList.Add(MakeShareable(new FDeviceDetailsFeature(TEXT("SdkConnectDisconnect"), TargetPlatform.SupportsFeature(ETargetPlatformFeatures::SdkConnectDisconnect))));
-			FeatureList.Add(MakeShareable(new FDeviceDetailsFeature(TEXT("Tessellation"), TargetPlatform.SupportsFeature(ETargetPlatformFeatures::Tessellation))));
-			FeatureList.Add(MakeShareable(new FDeviceDetailsFeature(TEXT("TextureStreaming"), TargetPlatform.SupportsFeature(ETargetPlatformFeatures::TextureStreaming))));
-			FeatureList.Add(MakeShareable(new FDeviceDetailsFeature(TEXT("UserCredentials"), TargetPlatform.SupportsFeature(ETargetPlatformFeatures::UserCredentials))));
-			FeatureList.Add(MakeShareable(new FDeviceDetailsFeature(TEXT("MobileRendering"), TargetPlatform.SupportsFeature(ETargetPlatformFeatures::MobileRendering))));
-			FeatureList.Add(MakeShareable(new FDeviceDetailsFeature(TEXT("DeferredRendering"), TargetPlatform.SupportsFeature(ETargetPlatformFeatures::DeferredRendering))));
-
-			// device features
-			FeatureList.Add(MakeShareable(new FDeviceDetailsFeature(TEXT("MultiLaunch"), TargetDevice->SupportsFeature(ETargetDeviceFeatures::MultiLaunch))));
-			FeatureList.Add(MakeShareable(new FDeviceDetailsFeature(TEXT("PowerOff"), TargetDevice->SupportsFeature(ETargetDeviceFeatures::PowerOff))));
-			FeatureList.Add(MakeShareable(new FDeviceDetailsFeature(TEXT("PowerOn"), TargetDevice->SupportsFeature(ETargetDeviceFeatures::PowerOn))));
-			FeatureList.Add(MakeShareable(new FDeviceDetailsFeature(TEXT("ProcessSnapshot"), TargetDevice->SupportsFeature(ETargetDeviceFeatures::ProcessSnapshot))));
-			FeatureList.Add(MakeShareable(new FDeviceDetailsFeature(TEXT("Reboot"), TargetDevice->SupportsFeature(ETargetDeviceFeatures::Reboot))));
+			ITargetDevicePtr TargetDevice = DeviceService->GetDevice();
+
+			if (TargetDevice.IsValid())
+			{
+				const ITargetPlatform& TargetPlatform = TargetDevice->GetTargetPlatform();
+
+				// platform features
+				FeatureList.Add(MakeShareable(new FDeviceDetailsFeature(TEXT("AudioStreaming"), TargetPlatform.SupportsFeature(ETargetPlatformFeatures::AudioStreaming))));
+				FeatureList.Add(MakeShareable(new FDeviceDetailsFeature(TEXT("DistanceFieldShadows"), TargetPlatform.SupportsFeature(ETargetPlatformFeatures::DistanceFieldShadows))));
+				FeatureList.Add(MakeShareable(new FDeviceDetailsFeature(TEXT("GrayscaleSRGB"), TargetPlatform.SupportsFeature(ETargetPlatformFeatures::GrayscaleSRGB))));
+				FeatureList.Add(MakeShareable(new FDeviceDetailsFeature(TEXT("HighQualityLightmaps"), TargetPlatform.SupportsFeature(ETargetPlatformFeatures::HighQualityLightmaps))));
+				FeatureList.Add(MakeShareable(new FDeviceDetailsFeature(TEXT("LowQualityLightmaps"), TargetPlatform.SupportsFeature(ETargetPlatformFeatures::LowQualityLightmaps))));
+				FeatureList.Add(MakeShareable(new FDeviceDetailsFeature(TEXT("MultipleGameInstances"), TargetPlatform.SupportsFeature(ETargetPlatformFeatures::MultipleGameInstances))));
+				FeatureList.Add(MakeShareable(new FDeviceDetailsFeature(TEXT("Packaging"), TargetPlatform.SupportsFeature(ETargetPlatformFeatures::Packaging))));
+				FeatureList.Add(MakeShareable(new FDeviceDetailsFeature(TEXT("SdkConnectDisconnect"), TargetPlatform.SupportsFeature(ETargetPlatformFeatures::SdkConnectDisconnect))));
+				FeatureList.Add(MakeShareable(new FDeviceDetailsFeature(TEXT("Tessellation"), TargetPlatform.SupportsFeature(ETargetPlatformFeatures::Tessellation))));
+				FeatureList.Add(MakeShareable(new FDeviceDetailsFeature(TEXT("TextureStreaming"), TargetPlatform.SupportsFeature(ETargetPlatformFeatures::TextureStreaming))));
+				FeatureList.Add(MakeShareable(new FDeviceDetailsFeature(TEXT("UserCredentials"), TargetPlatform.SupportsFeature(ETargetPlatformFeatures::UserCredentials))));
+				FeatureList.Add(MakeShareable(new FDeviceDetailsFeature(TEXT("MobileRendering"), TargetPlatform.SupportsFeature(ETargetPlatformFeatures::MobileRendering))));
+				FeatureList.Add(MakeShareable(new FDeviceDetailsFeature(TEXT("DeferredRendering"), TargetPlatform.SupportsFeature(ETargetPlatformFeatures::DeferredRendering))));
+
+				// device features
+				FeatureList.Add(MakeShareable(new FDeviceDetailsFeature(TEXT("MultiLaunch"), TargetDevice->SupportsFeature(ETargetDeviceFeatures::MultiLaunch))));
+				FeatureList.Add(MakeShareable(new FDeviceDetailsFeature(TEXT("PowerOff"), TargetDevice->SupportsFeature(ETargetDeviceFeatures::PowerOff))));
+				FeatureList.Add(MakeShareable(new FDeviceDetailsFeature(TEXT("PowerOn"), TargetDevice->SupportsFeature(ETargetDeviceFeatures::PowerOn))));
+				FeatureList.Add(MakeShareable(new FDeviceDetailsFeature(TEXT("ProcessSnapshot"), TargetDevice->SupportsFeature(ETargetDeviceFeatures::ProcessSnapshot))));
+				FeatureList.Add(MakeShareable(new FDeviceDetailsFeature(TEXT("Reboot"), TargetDevice->SupportsFeature(ETargetDeviceFeatures::Reboot))));
+			}
 		}
-	}
-
-	FeatureListView->RequestListRefresh();
-	QuickInfo->SetDeviceService(DeviceService);
-}
 
+		FeatureListView->RequestListRefresh();
+		QuickInfo->SetDeviceService(DeviceService);
+	};
 
-EVisibility SDeviceDetails::HandleSelectDeviceOverlayVisibility( ) const
-{
-	if (Model->GetSelectedDeviceService().IsValid())
-	{
-		return EVisibility::Hidden;
-	}
-
-	return EVisibility::Visible;
+	// wire up models
+	Model->OnSelectedDeviceServiceChanged().AddLambda(ModelSelectedDeviceServiceChanged);
 }
+END_SLATE_FUNCTION_BUILD_OPTIMIZATION
 
 
 #undef LOCTEXT_NAMESPACE
diff --git a/Engine/Source/Developer/DeviceManager/Private/Widgets/Details/SDeviceDetails.h b/Engine/Source/Developer/DeviceManager/Private/Widgets/Details/SDeviceDetails.h
--- a/Engine/Source/Developer/DeviceManager/Private/Widgets/Details/SDeviceDetails.h
+++ b/Engine/Source/Developer/DeviceManager/Private/Widgets/Details/SDeviceDetails.h
@@ -3,16 +3,18 @@
 #pragma once
 
 #include "CoreMinimal.h"
-#include "Layout/Visibility.h"
+#include "Containers/Array.h"
 #include "Widgets/DeclarativeSyntaxSupport.h"
 #include "Widgets/SCompoundWidget.h"
-#include "Models/DeviceManagerModel.h"
-#include "Widgets/Views/STableViewBase.h"
-#include "Widgets/Views/STableRow.h"
-#include "Models/DeviceDetailsFeature.h"
+#include "Templates/SharedPointer.h"
+#include "Widgets/Views/SListView.h"
 
+class FDeviceManagerModel;
 class SDeviceQuickInfo;
 
+struct FDeviceDetailsFeature;
+
+
 /**
  * Implements the device details widget.
  */
@@ -26,46 +28,30 @@ public:
 
 public:
 
-	/**
-	 * Destructor.
-	 */
-	~SDeviceDetails( );
+	/** Destructor. */
+	~SDeviceDetails();
 
 public:
 
 	/**
-	 * Constructs the widget.
+	 * Construct the widget.
 	 *
 	 * @param InArgs The construction arguments.
 	 * @param InModel The view model to use.
 	 */
-	void Construct( const FArguments& InArgs, const FDeviceManagerModelRef& InModel );
-
-private:
-
-	// Callback for getting the visibility of the details box.
-	EVisibility HandleDetailsBoxVisibility( ) const;
-
-	// Callback for generating a row widget for the feature list view.
-	TSharedRef<ITableRow> HandleFeatureListGenerateRow( FDeviceDetailsFeaturePtr Feature, const TSharedRef<STableViewBase>& OwnerTable );
-
-	// Callback for handling device service selection changes.
-	void HandleModelSelectedDeviceServiceChanged( );
-
-	// Callback for getting the visibility of the 'Select a device' message.
-	EVisibility HandleSelectDeviceOverlayVisibility( ) const;
+	void Construct(const FArguments& InArgs, const TSharedRef<FDeviceManagerModel>& InModel);
 
 private:
 
-	// Holds the list of device features.
-	TArray<FDeviceDetailsFeaturePtr> FeatureList;
+	/** The list of device features. */
+	TArray<TSharedPtr<FDeviceDetailsFeature>> FeatureList;
 
-	// Holds the device's feature list view.
-	TSharedPtr<SListView<FDeviceDetailsFeaturePtr> > FeatureListView;
+	/** The device's feature list view. */
+	TSharedPtr<SListView<TSharedPtr<FDeviceDetailsFeature>> > FeatureListView;
 
-	// Holds a pointer the device manager's view model.
-	FDeviceManagerModelPtr Model;
+	/** Pointer the device manager's view model. */
+	TSharedPtr<FDeviceManagerModel> Model;
 
-	// Holds the quick information widget.
+	/** The quick information widget. */
 	TSharedPtr<SDeviceQuickInfo> QuickInfo;
 };
diff --git a/Engine/Source/Developer/DeviceManager/Private/Widgets/Details/SDeviceDetailsFeatureListRow.h b/Engine/Source/Developer/DeviceManager/Private/Widgets/Details/SDeviceDetailsFeatureListRow.h
--- a/Engine/Source/Developer/DeviceManager/Private/Widgets/Details/SDeviceDetailsFeatureListRow.h
+++ b/Engine/Source/Developer/DeviceManager/Private/Widgets/Details/SDeviceDetailsFeatureListRow.h
@@ -2,7 +2,7 @@
 
 #pragma once
 
-#include "CoreMinimal.h"
+#include "CoreTypes.h"
 #include "Widgets/SNullWidget.h"
 #include "Widgets/DeclarativeSyntaxSupport.h"
 #include "Styling/SlateColor.h"
@@ -22,7 +22,7 @@
  * Implements a row widget for the device feature list.
  */
 class SDeviceDetailsFeatureListRow
-	: public SMultiColumnTableRow<FDeviceDetailsFeaturePtr>
+	: public SMultiColumnTableRow<TSharedPtr<FDeviceDetailsFeature>>
 {
 public:
 
@@ -38,22 +38,22 @@ public:
 	 * @param InArgs The construction arguments.
 	 * @param InOwnerTableView The table view that owns this row.
 	 */
-	void Construct( const FArguments& InArgs, const TSharedRef<STableViewBase>& InOwnerTableView, const FDeviceDetailsFeatureRef& InFeature )
+	void Construct(const FArguments& InArgs, const TSharedRef<STableViewBase>& InOwnerTableView, const TSharedRef<FDeviceDetailsFeature>& InFeature)
 	{
 //		check(InArgs._Style.IsValid());
 
 		Feature = InFeature;
 //		Style = InArgs._Style;
 
-		SMultiColumnTableRow<FDeviceDetailsFeaturePtr>::Construct(FSuperRowType::FArguments(), InOwnerTableView);
+		SMultiColumnTableRow<TSharedPtr<FDeviceDetailsFeature>>::Construct(FSuperRowType::FArguments(), InOwnerTableView);
 	}
 
 public:
 
-	// SMultiColumnTableRow interface
+	//~ SMultiColumnTableRow interface
 
 	BEGIN_SLATE_FUNCTION_BUILD_OPTIMIZATION
-	virtual TSharedRef<SWidget> GenerateWidgetForColumn( const FName& ColumnName ) override
+	virtual TSharedRef<SWidget> GenerateWidgetForColumn(const FName& ColumnName) override
 	{
 		if (ColumnName == "Available")
 		{
@@ -84,8 +84,8 @@ public:
 
 private:
 
-	// Callback for getting the text color.
-	FSlateColor HandleTextColorAndOpacity( ) const
+	/** Callback for getting the text color. */
+	FSlateColor HandleTextColorAndOpacity() const
 	{
 		if (Feature->Available)
 		{
@@ -97,10 +97,10 @@ private:
 
 private:
 
-	// Holds a pointer to the device feature to be shown in this row.
-	FDeviceDetailsFeaturePtr Feature;
+	/** Pointer to the device feature to be shown in this row. */
+	TSharedPtr<FDeviceDetailsFeature> Feature;
 
-	// Holds the widget's visual style.
+	/** The widget's visual style. */
 //	TSharedPtr<ISlateStyle> Style;
 };
 
diff --git a/Engine/Source/Developer/DeviceManager/Private/Widgets/Processes/SDeviceProcesses.cpp b/Engine/Source/Developer/DeviceManager/Private/Widgets/Processes/SDeviceProcesses.cpp
--- a/Engine/Source/Developer/DeviceManager/Private/Widgets/Processes/SDeviceProcesses.cpp
+++ b/Engine/Source/Developer/DeviceManager/Private/Widgets/Processes/SDeviceProcesses.cpp
@@ -1,12 +1,20 @@
 // Copyright 1998-2017 Epic Games, Inc. All Rights Reserved.
 
-#include "Widgets/Processes/SDeviceProcesses.h"
+#include "SDeviceProcesses.h"
+
+#include "EditorStyleSet.h"
+#include "Interfaces/ITargetDevice.h"
+#include "Internationalization/Text.h"
 #include "Misc/MessageDialog.h"
-#include "Widgets/SOverlay.h"
 #include "SlateOptMacros.h"
+#include "Widgets/SOverlay.h"
 #include "Widgets/Input/SButton.h"
 #include "Widgets/Input/SCheckBox.h"
-#include "EditorStyleSet.h"
+#include "Widgets/Views/STableViewBase.h"
+#include "Widgets/Views/STableRow.h"
+
+#include "Models/DeviceManagerModel.h"
+#include "Widgets/Processes/SDeviceProcessesProcessTreeNode.h"
 #include "Widgets/Processes/SDeviceProcessesProcessListRow.h"
 
 
@@ -16,7 +24,7 @@
 /* SMessagingEndpoints structors
 *****************************************************************************/
 
-SDeviceProcesses::~SDeviceProcesses( )
+SDeviceProcesses::~SDeviceProcesses()
 {
 	if (Model.IsValid())
 	{
@@ -29,11 +37,53 @@ SDeviceProcesses::~SDeviceProcesses( )
  *****************************************************************************/
 
 BEGIN_SLATE_FUNCTION_BUILD_OPTIMIZATION
-void SDeviceProcesses::Construct( const FArguments& InArgs, const FDeviceManagerModelRef& InModel )
+void SDeviceProcesses::Construct(const FArguments& InArgs, const TSharedRef<FDeviceManagerModel>& InModel)
 {
 	Model = InModel;
 	ShowProcessTree = true;
 
+	// callback for getting the text of the message overlay.
+	auto MessageOverlayText = [this]() -> FText {
+		ITargetDeviceServicePtr DeviceService = Model->GetSelectedDeviceService();
+
+		if (DeviceService.IsValid())
+		{
+			ITargetDevicePtr Device = DeviceService->GetDevice();
+
+			if (Device.IsValid() && Device->IsConnected())
+			{
+				if (Device->SupportsFeature(ETargetDeviceFeatures::ProcessSnapshot))
+				{
+					return FText::GetEmpty();
+				}
+
+				return LOCTEXT("ProcessSnapshotsUnsupportedOverlayText", "The selected device does not support process snapshots");
+			}
+
+			return LOCTEXT("DeviceUnavailableOverlayText", "The selected device is currently unavailable");
+		}
+
+		return LOCTEXT("SelectDeviceOverlayText", "Please select a device from the Device Browser");
+	};
+
+	// callback for getting the visibility of the message overlay.
+	auto MessageOverlayVisibility = [this]() -> EVisibility {
+		ITargetDeviceServicePtr DeviceService = Model->GetSelectedDeviceService();
+
+		if (DeviceService.IsValid())
+		{
+			ITargetDevicePtr Device = DeviceService->GetDevice();
+
+			if (Device.IsValid() && Device->IsConnected() && Device->SupportsFeature(ETargetDeviceFeatures::ProcessSnapshot))
+			{
+				return EVisibility::Hidden;
+			}
+		}
+
+		return EVisibility::Visible;
+	};
+
+	// construct children
 	ChildSlot
 	[
 		SNew(SOverlay)
@@ -52,10 +102,22 @@ void SDeviceProcesses::Construct( const FArguments& InArgs, const FDeviceManager
 							.BorderImage(FEditorStyle::GetBrush("ToolPanel.GroupBorder"))
 							.Padding(0.0f)
 							[
-								SAssignNew(ProcessTreeView, STreeView<FDeviceProcessesProcessTreeNodePtr>)
+								SAssignNew(ProcessTreeView, STreeView<TSharedPtr<FDeviceProcessesProcessTreeNode>>)
 									.ItemHeight(20.0f)
-									.OnGenerateRow(this, &SDeviceProcesses::HandleProcessTreeViewGenerateRow)
-									.OnGetChildren(this, &SDeviceProcesses::HandleProcessTreeViewGetChildren)
+									.OnGenerateRow_Lambda(
+										[](TSharedPtr<FDeviceProcessesProcessTreeNode> Item, const TSharedRef<STableViewBase>& OwnerTable) -> TSharedRef<ITableRow> {
+											return SNew(SDeviceProcessesProcessListRow, OwnerTable)
+												.Node(Item);
+										}
+									)
+									.OnGetChildren_Lambda(
+										[](TSharedPtr<FDeviceProcessesProcessTreeNode> Item, TArray<TSharedPtr<FDeviceProcessesProcessTreeNode>>& OutChildren) {
+											if (Item.IsValid())
+											{
+												OutChildren = Item->GetChildren();
+											}
+										}
+									)
 									.SelectionMode(ESelectionMode::Multi)
 									.TreeItemsSource(&ProcessList)
 									.HeaderRow
@@ -98,8 +160,17 @@ void SDeviceProcesses::Construct( const FArguments& InArgs, const FDeviceManager
 								+ SHorizontalBox::Slot()
 									[
 										SNew(SCheckBox)
-											.IsChecked(this, &SDeviceProcesses::HandleProcessTreeCheckBoxIsChecked)
-											.OnCheckStateChanged(this, &SDeviceProcesses::HandleProcessTreeCheckBoxCheckStateChanged)
+											.IsChecked_Lambda(
+												[this]() -> ECheckBoxState {
+													return ShowProcessTree ? ECheckBoxState::Checked : ECheckBoxState::Unchecked;
+												}
+											)
+											.OnCheckStateChanged_Lambda(
+												[this](ECheckBoxState NewState) {
+													ShowProcessTree = (NewState == ECheckBoxState::Checked);
+													ReloadProcessList(false);
+												}
+											)
 											.Padding(FMargin(4.0f, 0.0f))
 											.ToolTipText(LOCTEXT("ProcessTreeCheckBoxToolTip", "Check this box to display the list of processes as a tree instead of a flat list"))
 											[
@@ -112,7 +183,11 @@ void SDeviceProcesses::Construct( const FArguments& InArgs, const FDeviceManager
 									.HAlign(HAlign_Right)
 									[
 										SNew(SButton)
-											.IsEnabled(this, &SDeviceProcesses::HandleTerminateProcessButtonIsEnabled)
+											.IsEnabled_Lambda(
+												[this]() -> bool {
+													return (ProcessTreeView->GetNumItemsSelected() > 0);
+												}
+											)
 											.OnClicked(this, &SDeviceProcesses::HandleTerminateProcessButtonClicked)
 											.Text(LOCTEXT("TerminateProcessButtonText", "Terminate Process"))
 									]
@@ -127,27 +202,25 @@ void SDeviceProcesses::Construct( const FArguments& InArgs, const FDeviceManager
 				SNew(SBorder)
 					.BorderImage(FEditorStyle::GetBrush("NotificationList.ItemBackground"))
 					.Padding(8.0f)
-					.Visibility(this, &SDeviceProcesses::HandleMessageOverlayVisibility)
+					.Visibility_Lambda(MessageOverlayVisibility)
 					[
 						SNew(STextBlock)
-							.Text(this, &SDeviceProcesses::HandleMessageOverlayText)
+							.Text_Lambda(MessageOverlayText)
 					]
 			]
 	];
 
-	Model->OnSelectedDeviceServiceChanged().AddRaw(this, &SDeviceProcesses::HandleModelSelectedDeviceServiceChanged);
+	// callback for handling device service selection changes.
+	auto ModelSelectedDeviceServiceChanged = [this]() {
+		ReloadProcessList(true);
+	};
 
-	ReloadProcessList(true);
+	// wire up models
+	Model->OnSelectedDeviceServiceChanged().AddLambda([this]() {  });
 
-	// Register for an active update every 2.5 seconds
-	RegisterActiveTimer( 2.5f, FWidgetActiveTimerDelegate::CreateSP( this, &SDeviceProcesses::UpdateProcessList ) );
-}
-
-EActiveTimerReturnType SDeviceProcesses::UpdateProcessList( double InCurrentTime, float InDeltaTime )
-{
-	ReloadProcessList( true );
-
-	return EActiveTimerReturnType::Continue;
+	// initialize
+	ReloadProcessList(true);
+	RegisterActiveTimer(2.5f, FWidgetActiveTimerDelegate::CreateSP(this, &SDeviceProcesses::UpdateProcessList));
 }
 END_SLATE_FUNCTION_BUILD_OPTIMIZATION
 
@@ -155,7 +228,7 @@ END_SLATE_FUNCTION_BUILD_OPTIMIZATION
 /* SDeviceDetails implementation
  *****************************************************************************/
 
-void SDeviceProcesses::ReloadProcessList( bool FullyReload )
+void SDeviceProcesses::ReloadProcessList(bool FullyReload)
 {
 	// reload running processes
 	if (FullyReload)
@@ -176,12 +249,12 @@ void SDeviceProcesses::ReloadProcessList( bool FullyReload )
 	}
 
 	// update process tree
-	TMap<uint32, FDeviceProcessesProcessTreeNodePtr> NewProcessMap;
+	TMap<uint32, TSharedPtr<FDeviceProcessesProcessTreeNode>> NewProcessMap;
 
 	for (int32 ProcessIndex = 0; ProcessIndex < RunningProcesses.Num(); ++ProcessIndex)
 	{
 		const FTargetDeviceProcessInfo& ProcessInfo = RunningProcesses[ProcessIndex];
-		FDeviceProcessesProcessTreeNodePtr Node = ProcessMap.FindRef(ProcessInfo.Id);
+		TSharedPtr<FDeviceProcessesProcessTreeNode> Node = ProcessMap.FindRef(ProcessInfo.Id);
 
 		if (Node.IsValid())
 		{
@@ -201,10 +274,10 @@ void SDeviceProcesses::ReloadProcessList( bool FullyReload )
 	// build process tree
 	if (ShowProcessTree)
 	{
-		for (TMap<uint32, FDeviceProcessesProcessTreeNodePtr>::TConstIterator It(ProcessMap); It; ++It)
+		for (TMap<uint32, TSharedPtr<FDeviceProcessesProcessTreeNode>>::TConstIterator It(ProcessMap); It; ++It)
 		{
-			FDeviceProcessesProcessTreeNodePtr Node = It.Value();
-			FDeviceProcessesProcessTreeNodePtr Parent = ProcessMap.FindRef(Node->GetProcessInfo().ParentId);
+			TSharedPtr<FDeviceProcessesProcessTreeNode> Node = It.Value();
+			TSharedPtr<FDeviceProcessesProcessTreeNode> Parent = ProcessMap.FindRef(Node->GetProcessInfo().ParentId);
 
 			if (Parent.IsValid())
 			{
@@ -217,9 +290,9 @@ void SDeviceProcesses::ReloadProcessList( bool FullyReload )
 	// filter process list
 	ProcessList.Reset();
 
-	for (TMap<uint32, FDeviceProcessesProcessTreeNodePtr>::TConstIterator It(ProcessMap); It; ++It)
+	for (TMap<uint32, TSharedPtr<FDeviceProcessesProcessTreeNode>>::TConstIterator It(ProcessMap); It; ++It)
 	{
-		FDeviceProcessesProcessTreeNodePtr Node = It.Value();
+		TSharedPtr<FDeviceProcessesProcessTreeNode> Node = It.Value();
 
 		if (!Node->GetParent().IsValid())
 		{
@@ -234,94 +307,18 @@ void SDeviceProcesses::ReloadProcessList( bool FullyReload )
 }
 
 
-/* SDeviceDetails callbacks
- *****************************************************************************/
-
-FText SDeviceProcesses::HandleMessageOverlayText( ) const
-{
-	ITargetDeviceServicePtr DeviceService = Model->GetSelectedDeviceService();
-
-	if (DeviceService.IsValid())
-	{
-		ITargetDevicePtr Device = DeviceService->GetDevice();
-
-		if (Device.IsValid() && Device->IsConnected())
-		{
-			if (Device->SupportsFeature(ETargetDeviceFeatures::ProcessSnapshot))
-			{
-				return FText::GetEmpty();
-			}
-
-			return LOCTEXT("ProcessSnapshotsUnsupportedOverlayText", "The selected device does not support process snapshots");
-		}
-
-		return LOCTEXT("DeviceUnavailableOverlayText", "The selected device is currently unavailable");
-	}
-
-	return LOCTEXT("SelectDeviceOverlayText", "Please select a device from the Device Browser");
-}
-
-
-EVisibility SDeviceProcesses::HandleMessageOverlayVisibility( ) const
-{
-	ITargetDeviceServicePtr DeviceService = Model->GetSelectedDeviceService();
-
-	if (DeviceService.IsValid())
-	{
-		ITargetDevicePtr Device = DeviceService->GetDevice();
-
-		if (Device.IsValid() && Device->IsConnected() && Device->SupportsFeature(ETargetDeviceFeatures::ProcessSnapshot))
-		{
-			return EVisibility::Hidden;
-		}
-	}
-
-	return EVisibility::Visible;
-}
-
-
-void SDeviceProcesses::HandleModelSelectedDeviceServiceChanged( )
+EActiveTimerReturnType SDeviceProcesses::UpdateProcessList(double InCurrentTime, float InDeltaTime)
 {
 	ReloadProcessList(true);
-}
-
 
-void SDeviceProcesses::HandleProcessTreeCheckBoxCheckStateChanged( ECheckBoxState NewState )
-{
-	ShowProcessTree = (NewState == ECheckBoxState::Checked);
-
-	ReloadProcessList(false);
-}
-
-
-ECheckBoxState SDeviceProcesses::HandleProcessTreeCheckBoxIsChecked( ) const
-{
-	if (ShowProcessTree)
-	{
-		return ECheckBoxState::Checked;
-	}
-
-	return ECheckBoxState::Unchecked;
-}
-
-
-TSharedRef<ITableRow> SDeviceProcesses::HandleProcessTreeViewGenerateRow( FDeviceProcessesProcessTreeNodePtr Item, const TSharedRef<STableViewBase>& OwnerTable )
-{
-	return SNew(SDeviceProcessesProcessListRow, OwnerTable)
-		.Node(Item);
+	return EActiveTimerReturnType::Continue;
 }
 
 
-void SDeviceProcesses::HandleProcessTreeViewGetChildren( FDeviceProcessesProcessTreeNodePtr Item, TArray<FDeviceProcessesProcessTreeNodePtr>& OutChildren )
-{
-	if (Item.IsValid())
-	{
-		OutChildren = Item->GetChildren();
-	}
-}
-
+/* SDeviceDetails callbacks
+ *****************************************************************************/
 
-bool SDeviceProcesses::HandleProcessesBoxIsEnabled( ) const
+bool SDeviceProcesses::HandleProcessesBoxIsEnabled() const
 {
 	ITargetDeviceServicePtr DeviceService = Model->GetSelectedDeviceService();
 
@@ -336,13 +333,7 @@ bool SDeviceProcesses::HandleProcessesBoxIsEnabled( ) const
 }
 
 
-bool SDeviceProcesses::HandleTerminateProcessButtonIsEnabled( ) const
-{
-	return (ProcessTreeView->GetNumItemsSelected() > 0);
-}
-
-
-FReply SDeviceProcesses::HandleTerminateProcessButtonClicked( )
+FReply SDeviceProcesses::HandleTerminateProcessButtonClicked()
 {
 	ITargetDeviceServicePtr DeviceService = Model->GetSelectedDeviceService();
 
@@ -350,8 +341,8 @@ FReply SDeviceProcesses::HandleTerminateProcessButtonClicked( )
 	{
 		if (FMessageDialog::Open(EAppMsgType::OkCancel, LOCTEXT("TerminateProcessWarning", "Warning: If you terminate a process that is associated with a game or an application, you willl lose any unsaved data. If you end a system process, it might result in an unstable system.")) == EAppReturnType::Ok)
 		{
-			TArray<FDeviceProcessesProcessTreeNodePtr> FailedProcesses;
-			TArray<FDeviceProcessesProcessTreeNodePtr> SelectedProcesses = ProcessTreeView->GetSelectedItems();
+			TArray<TSharedPtr<FDeviceProcessesProcessTreeNode>> FailedProcesses;
+			TArray<TSharedPtr<FDeviceProcessesProcessTreeNode>> SelectedProcesses = ProcessTreeView->GetSelectedItems();
 
 			for (int32 ProcessIndex = 0; ProcessIndex < SelectedProcesses.Num(); ++ProcessIndex)
 			{
@@ -359,7 +350,7 @@ FReply SDeviceProcesses::HandleTerminateProcessButtonClicked( )
 
 				if (TargetDevice.IsValid())
 				{
-					const FDeviceProcessesProcessTreeNodePtr& Process = SelectedProcesses[ProcessIndex];
+					const TSharedPtr<FDeviceProcessesProcessTreeNode>& Process = SelectedProcesses[ProcessIndex];
 
 					if (!TargetDevice->TerminateProcess(Process->GetProcessInfo().Id))
 					{
@@ -377,7 +368,7 @@ FReply SDeviceProcesses::HandleTerminateProcessButtonClicked( )
 					ProcessInfo += FString::Printf(TEXT("%s (PID: %d)\n"), *FailedProcessInfo.Name, FailedProcessInfo.Id);
 				}
 
-				const FText ErrorMessage = FText::Format( LOCTEXT("FailedToTerminateProcessesMessage", "The following processes could not be terminated.\nYou may not have the required permissions:\n\n{0}"), FText::FromString( ProcessInfo ) );
+				const FText ErrorMessage = FText::Format(LOCTEXT("FailedToTerminateProcessesMessage", "The following processes could not be terminated.\nYou may not have the required permissions:\n\n{0}"), FText::FromString(ProcessInfo));
 				FMessageDialog::Open(EAppMsgType::Ok, ErrorMessage);
 			}
 		}
diff --git a/Engine/Source/Developer/DeviceManager/Private/Widgets/Processes/SDeviceProcesses.h b/Engine/Source/Developer/DeviceManager/Private/Widgets/Processes/SDeviceProcesses.h
--- a/Engine/Source/Developer/DeviceManager/Private/Widgets/Processes/SDeviceProcesses.h
+++ b/Engine/Source/Developer/DeviceManager/Private/Widgets/Processes/SDeviceProcesses.h
@@ -2,17 +2,22 @@
 
 #pragma once
 
-#include "CoreMinimal.h"
+#include "CoreTypes.h"
+#include "Containers/Array.h"
+#include "Containers/Map.h"
+#include "Input/Reply.h"
 #include "Layout/Visibility.h"
+#include "Misc/DateTime.h"
+#include "Templates/SharedPointer.h"
 #include "Widgets/DeclarativeSyntaxSupport.h"
-#include "Input/Reply.h"
 #include "Widgets/SCompoundWidget.h"
-#include "Interfaces/ITargetDevice.h"
-#include "Models/DeviceManagerModel.h"
-#include "Widgets/Views/STableViewBase.h"
-#include "Widgets/Views/STableRow.h"
 #include "Widgets/Views/STreeView.h"
-#include "Widgets/Processes/SDeviceProcessesProcessTreeNode.h"
+
+class FDeviceManagerModel;
+class FDeviceProcessesProcessTreeNode;
+
+struct FTargetDeviceProcessInfo;
+
 
 /**
  * Implements the device details widget.
@@ -27,82 +32,59 @@ public:
 
 public:
 
-	/**
-	 * Destructor.
-	 */
-	~SDeviceProcesses( );
+	/** Destructor. */
+	~SDeviceProcesses();
 
 public:
 
 	/**
-	 * Constructs the widget.
+	 * Construct the widget.
 	 *
 	 * @param InArgs The construction arguments.
 	 * @param InModel The view model to use.
 	 */
-	void Construct( const FArguments& InArgs, const FDeviceManagerModelRef& InModel );
+	void Construct(const FArguments& InArgs, const TSharedRef<FDeviceManagerModel>& InModel);
 
 protected:
 
 	/**
-	 * Reloads the list of processes.
+	 * Reload the list of processes.
+	 *
+	 * @param FullyReload true reload all entries, false to refresh the existing entries.
 	 */
-	void ReloadProcessList( bool FullyReload );
+	void ReloadProcessList(bool FullyReload);
 
-private:
-	// Periodically refreshes the process list
-	EActiveTimerReturnType UpdateProcessList( double InCurrentTime, float InDeltaTime );
+	/** Periodically refreshes the process list. */
+	EActiveTimerReturnType UpdateProcessList(double InCurrentTime, float InDeltaTime);
 
-	// Callback for getting the text of the message overlay.
-	FText HandleMessageOverlayText( ) const;
-
-	// Callback for getting the visibility of the message overlay.
-	EVisibility HandleMessageOverlayVisibility( ) const;
-
-	// Callback for handling device service selection changes.
-	void HandleModelSelectedDeviceServiceChanged( );
-
-	// Callback for changing 'Show process tree' check box state.
-	void HandleProcessTreeCheckBoxCheckStateChanged( ECheckBoxState NewState );
-
-	// Callback for determining checked state of the 'Show process tree' check box.
-	ECheckBoxState HandleProcessTreeCheckBoxIsChecked( ) const;
-
-	// Callback for generating a row widget in the process tree view.
-	TSharedRef<ITableRow> HandleProcessTreeViewGenerateRow( FDeviceProcessesProcessTreeNodePtr Item, const TSharedRef<STableViewBase>& OwnerTable );
-
-	// Callback for getting the children of a node in the process tree view.
-	void HandleProcessTreeViewGetChildren( FDeviceProcessesProcessTreeNodePtr Item, TArray<FDeviceProcessesProcessTreeNodePtr>& OutChildren );
-
-	// Callback for getting the enabled state of the processes panel.
-	bool HandleProcessesBoxIsEnabled( ) const;
+private:
 
-	// Callback for getting the enabled state of the 'Terminate Process' button.
-	bool HandleTerminateProcessButtonIsEnabled( ) const;
+	/** Callback for getting the enabled state of the processes panel. */
+	bool HandleProcessesBoxIsEnabled() const;
 
-	// Callback for clicking the 'Terminate Process' button.
-	FReply HandleTerminateProcessButtonClicked( );
+	/** Callback for clicking the 'Terminate Process' button. */
+	FReply HandleTerminateProcessButtonClicked();
 
 private:
 	
-	// Holds the time at which the process list was last refreshed.
+	/** Holds the time at which the process list was last refreshed. */
 	FDateTime LastProcessListRefreshTime;
 
-	// Holds a pointer the device manager's view model.
-	FDeviceManagerModelPtr Model;
+	/** Holds a pointer the device manager's view model. */
+	TSharedPtr<FDeviceManagerModel> Model;
 
-	// Holds the filtered list of processes running on the device.
-	TArray<FDeviceProcessesProcessTreeNodePtr> ProcessList;
+	/** Holds the filtered list of processes running on the device. */
+	TArray<TSharedPtr<FDeviceProcessesProcessTreeNode>> ProcessList;
 
-	// Holds the collection of process nodes.
-	TMap<uint32, FDeviceProcessesProcessTreeNodePtr> ProcessMap;
+	/** Holds the collection of process nodes. */
+	TMap<uint32, TSharedPtr<FDeviceProcessesProcessTreeNode>> ProcessMap;
 
-	// Holds the application tree view.
-	TSharedPtr<STreeView<FDeviceProcessesProcessTreeNodePtr> > ProcessTreeView;
+	/** Holds the application tree view. */
+	TSharedPtr<STreeView<TSharedPtr<FDeviceProcessesProcessTreeNode>>> ProcessTreeView;
 
-	// Holds the collection of processes running on the device.
+	/** Holds the collection of processes running on the device. */
 	TArray<FTargetDeviceProcessInfo> RunningProcesses;
 
-	// Holds a flag indicating whether the process list should be shown as a tree instead of a list.
+	/** Holds a flag indicating whether the process list should be shown as a tree instead of a list. */
 	bool ShowProcessTree;
 };
diff --git a/Engine/Source/Developer/DeviceManager/Private/Widgets/Processes/SDeviceProcessesProcessListRow.h b/Engine/Source/Developer/DeviceManager/Private/Widgets/Processes/SDeviceProcessesProcessListRow.h
--- a/Engine/Source/Developer/DeviceManager/Private/Widgets/Processes/SDeviceProcessesProcessListRow.h
+++ b/Engine/Source/Developer/DeviceManager/Private/Widgets/Processes/SDeviceProcessesProcessListRow.h
@@ -2,60 +2,59 @@
 
 #pragma once
 
-#include "CoreMinimal.h"
-#include "Widgets/SNullWidget.h"
-#include "Widgets/DeclarativeSyntaxSupport.h"
-#include "Widgets/SWidget.h"
+#include "CoreTypes.h"
 #include "Layout/Margin.h"
-#include "Widgets/Views/STableViewBase.h"
+#include "SlateOptMacros.h"
+#include "Templates/SharedPointer.h"
+#include "Widgets/DeclarativeSyntaxSupport.h"
 #include "Widgets/SBoxPanel.h"
+#include "Widgets/SNullWidget.h"
+#include "Widgets/SWidget.h"
 #include "Widgets/Layout/SBox.h"
-#include "Widgets/Views/SExpanderArrow.h"
-#include "Widgets/Views/STableRow.h"
 #include "Widgets/Text/STextBlock.h"
+#include "Widgets/Views/SExpanderArrow.h"
 #include "Widgets/Views/SListView.h"
+#include "Widgets/Views/STableRow.h"
+#include "Widgets/Views/STableViewBase.h"
+
 #include "Widgets/Processes/SDeviceProcessesProcessTreeNode.h"
-#include "SlateOptMacros.h"
+
 
 #define LOCTEXT_NAMESPACE "SDeviceProcessesProcessListRow"
 
+
 /**
  * Implements a row widget for the process list view.
  */
 class SDeviceProcessesProcessListRow
-	: public SMultiColumnTableRow<FDeviceProcessesProcessTreeNodePtr>
+	: public SMultiColumnTableRow<TSharedPtr<FDeviceProcessesProcessTreeNode>>
 {
 public:
 
 	SLATE_BEGIN_ARGS(SDeviceProcessesProcessListRow) { }
-		SLATE_ARGUMENT(FDeviceProcessesProcessTreeNodePtr, Node)
+		SLATE_ARGUMENT(TSharedPtr<FDeviceProcessesProcessTreeNode>, Node)
 	SLATE_END_ARGS()
 
 public:
 
 	/**
-	 * Constructs the widget.
+	 * Construct the widget.
 	 *
 	 * @param InArgs The arguments.
 	 */
-	void Construct( const FArguments& InArgs, const TSharedRef<STableViewBase>& InOwnerTableView )
+	void Construct(const FArguments& InArgs, const TSharedRef<STableViewBase>& InOwnerTableView)
 	{
 		Node = InArgs._Node;
 
-		SMultiColumnTableRow<FDeviceProcessesProcessTreeNodePtr>::Construct(FSuperRowType::FArguments(), InOwnerTableView);
+		SMultiColumnTableRow<TSharedPtr<FDeviceProcessesProcessTreeNode>>::Construct(FSuperRowType::FArguments(), InOwnerTableView);
 	}
 
 public:
 
-	/**
-	 * Generates the widget for the specified column.
-	 *
-	 * @param ColumnName The name of the column to generate the widget for.
-	 *
-	 * @return The widget.
-	 */
+	//~ SMultiColumnTableRow interface
+
 	BEGIN_SLATE_FUNCTION_BUILD_OPTIMIZATION
-	virtual TSharedRef<SWidget> GenerateWidgetForColumn( const FName& ColumnName ) override
+	virtual TSharedRef<SWidget> GenerateWidgetForColumn(const FName& ColumnName) override
 	{
 		if (ColumnName == TEXT("Name"))
 		{
@@ -137,8 +136,8 @@ public:
 
 private:
 
-	// Holds the name of the process.
-	FDeviceProcessesProcessTreeNodePtr Node;
+	/** The name of the process. */
+	TSharedPtr<FDeviceProcessesProcessTreeNode> Node;
 };
 
 
diff --git a/Engine/Source/Developer/DeviceManager/Private/Widgets/Processes/SDeviceProcessesProcessTreeNode.h b/Engine/Source/Developer/DeviceManager/Private/Widgets/Processes/SDeviceProcessesProcessTreeNode.h
--- a/Engine/Source/Developer/DeviceManager/Private/Widgets/Processes/SDeviceProcessesProcessTreeNode.h
+++ b/Engine/Source/Developer/DeviceManager/Private/Widgets/Processes/SDeviceProcessesProcessTreeNode.h
@@ -2,17 +2,13 @@
 
 #pragma once
 
-#include "CoreMinimal.h"
+#include "CoreTypes.h"
+#include "Containers/Array.h"
 #include "Interfaces/ITargetDevice.h"
+#include "Templates/SharedPointer.h"
 
 class FDeviceProcessesProcessTreeNode;
 
-/** Type definition for shared pointers to instances of FDeviceProcessesProcessTreeNode. */
-typedef TSharedPtr<class FDeviceProcessesProcessTreeNode> FDeviceProcessesProcessTreeNodePtr;
-
-/** Type definition for shared references to instances of FDeviceProcessesProcessTreeNode. */
-typedef TSharedRef<class FDeviceProcessesProcessTreeNode> FDeviceProcessesProcessTreeNodeRef;
-
 
 /**
  * Implements a node for the process tree view.
@@ -22,82 +18,88 @@ class FDeviceProcessesProcessTreeNode
 public:
 
 	/**
-	 * Creates and initializes a new instance.
+	 * Create and initializes a new instance.
 	 *
 	 * @param InProcessInfo The node's process information.
 	 */
-	FDeviceProcessesProcessTreeNode( const FTargetDeviceProcessInfo& InProcessInfo )
+	FDeviceProcessesProcessTreeNode(const FTargetDeviceProcessInfo& InProcessInfo)
 		: ProcessInfo(InProcessInfo)
 	{ }
 
 public:
 
 	/**
-	 * Adds a child process node to this node.
+	 * Add a child process node to this node.
 	 *
 	 * @param The child node to add.
+	 * @see ClearChildren, GetChildren
 	 */
-	void AddChild( const FDeviceProcessesProcessTreeNodePtr& Child )
+	void AddChild(const TSharedPtr<FDeviceProcessesProcessTreeNode>& Child)
 	{
 		Children.Add(Child);
 	}
 
 	/**
-	 * Clears the collection of child nodes.
+	 * Clear the collection of child nodes.
+	 *
+	 * @see AddChild, GetChildren
 	 */
-	void ClearChildren( )
+	void ClearChildren()
 	{
 		Children.Reset();
 	}
 
 	/**
-	 * Gets the child nodes.
+	 * Get the child nodes.
 	 *
 	 * @return Child nodes.
+	 * @see GetParent
 	 */
-	const TArray<FDeviceProcessesProcessTreeNodePtr>& GetChildren( )
+	const TArray<TSharedPtr<FDeviceProcessesProcessTreeNode>>& GetChildren()
 	{
 		return Children;
 	}
 
 	/**
-	 * Gets the parent node.
+	 * Get the parent node.
 	 *
 	 * @return Parent node.
+	 * @see GetChildren, SetParent
 	 */
-	const FDeviceProcessesProcessTreeNodePtr& GetParent( )
+	const TSharedPtr<FDeviceProcessesProcessTreeNode>& GetParent()
 	{
 		return Parent;
 	}
 
 	/**
-	 * Gets the node's process information.
+	 * Get the node's process information.
 	 *
 	 * @return The process information.
 	 */
-	const FTargetDeviceProcessInfo& GetProcessInfo( ) const
+	const FTargetDeviceProcessInfo& GetProcessInfo() const
 	{
 		return ProcessInfo;
 	}
 
 	/**
-	 * Sets the parent node.
+	 * Set the parent node.
 	 *
 	 * @param Node The parent node to set.
+	 * @see GetParent
 	 */
-	void SetParent( const FDeviceProcessesProcessTreeNodePtr& Node )
+	void SetParent(const TSharedPtr<FDeviceProcessesProcessTreeNode>& Node)
 	{
 		Parent = Node;
 	}
 
 private:
 
-	// Holds the child process nodes
-	TArray<FDeviceProcessesProcessTreeNodePtr> Children;
+	/** The child process nodes. */
+	TArray<TSharedPtr<FDeviceProcessesProcessTreeNode>> Children;
 
-	// Holds a pointer to the parent node.
-	FDeviceProcessesProcessTreeNodePtr Parent;
+	/** Pointer to the parent node. */
+	TSharedPtr<FDeviceProcessesProcessTreeNode> Parent;
 
-	// Holds the process information
+	/** The process information. */
 	FTargetDeviceProcessInfo ProcessInfo;
 };
