#ue5 - Fix type conversion warnings in DeviceManager, DataTableEditor modules, and some related Slate files.

diff --git a/Engine/Source/Developer/DeviceManager/Private/Widgets/Browser/SDeviceBrowserDeviceAdder.cpp b/Engine/Source/Developer/DeviceManager/Private/Widgets/Browser/SDeviceBrowserDeviceAdder.cpp
--- a/Engine/Source/Developer/DeviceManager/Private/Widgets/Browser/SDeviceBrowserDeviceAdder.cpp
+++ b/Engine/Source/Developer/DeviceManager/Private/Widgets/Browser/SDeviceBrowserDeviceAdder.cpp
@@ -71,8 +71,8 @@ void SDeviceBrowserDeviceAdder::Construct(const FArguments& InArgs, const TShare
 				.HAlign(HAlign_Left)
 				[
 					SNew(SBox)
-						.WidthOverride(24)
-						.HeightOverride(24)
+						.WidthOverride(24.f)
+						.HeightOverride(24.f)
 						[
 							SNew(SImage)
 								.Image((PlatformInfo) ? FAppStyle::GetBrush(PlatformInfo->GetIconStyleName(EPlatformIconSize::Normal)) : FStyleDefaults::GetNoBrush())
diff --git a/Engine/Source/Developer/DeviceManager/Private/Widgets/Browser/SDeviceBrowserDeviceListRow.h b/Engine/Source/Developer/DeviceManager/Private/Widgets/Browser/SDeviceBrowserDeviceListRow.h
--- a/Engine/Source/Developer/DeviceManager/Private/Widgets/Browser/SDeviceBrowserDeviceListRow.h
+++ b/Engine/Source/Developer/DeviceManager/Private/Widgets/Browser/SDeviceBrowserDeviceListRow.h
@@ -86,8 +86,8 @@ public:
 
 			return SNew(SBox)
 				.Padding(FMargin(4.0f, 0.0f))
-				.WidthOverride(24)
-				.HeightOverride(24)
+				.WidthOverride(24.f)
+				.HeightOverride(24.f)
 				[
 					SNew(SImage)
 						.Image((PlatformInfo) ? FAppStyle::GetBrush(PlatformInfo->GetIconStyleName(EPlatformIconSize::Normal)) : FStyleDefaults::GetNoBrush())
diff --git a/Engine/Source/Developer/DeviceManager/Private/Widgets/Browser/SDeviceBrowserFilterBar.cpp b/Engine/Source/Developer/DeviceManager/Private/Widgets/Browser/SDeviceBrowserFilterBar.cpp
--- a/Engine/Source/Developer/DeviceManager/Private/Widgets/Browser/SDeviceBrowserFilterBar.cpp
+++ b/Engine/Source/Developer/DeviceManager/Private/Widgets/Browser/SDeviceBrowserFilterBar.cpp
@@ -77,8 +77,8 @@ void SDeviceBrowserFilterBar::Construct(const FArguments& InArgs, TSharedRef<FDe
 						.AutoWidth()
 						[
 							SNew(SBox)
-								.WidthOverride(24)
-								.HeightOverride(24)
+								.WidthOverride(24.f)
+								.HeightOverride(24.f)
 								[
 									SNew(SImage)
 									.Image((PlatformInfo) ? FAppStyle::GetBrush(PlatformInfo->GetIconStyleName(EPlatformIconSize::Normal)) : FStyleDefaults::GetNoBrush())
diff --git a/Engine/Source/Developer/DeviceManager/Private/Widgets/Processes/SDeviceProcesses.cpp b/Engine/Source/Developer/DeviceManager/Private/Widgets/Processes/SDeviceProcesses.cpp
--- a/Engine/Source/Developer/DeviceManager/Private/Widgets/Processes/SDeviceProcesses.cpp
+++ b/Engine/Source/Developer/DeviceManager/Private/Widgets/Processes/SDeviceProcesses.cpp
@@ -271,7 +271,7 @@ void SDeviceProcesses::ReloadProcessList(bool FullyReload)
 void SDeviceProcesses::UpdateProcessTree()
 {
 	// update process tree
-	TMap<uint32, TSharedPtr<FDeviceProcessesProcessTreeNode>> NewProcessMap;
+	SDeviceProcesses::ProcessMapType NewProcessMap;
 
 	for (int32 ProcessIndex = 0; ProcessIndex < RunningProcesses.Num(); ++ProcessIndex)
 	{
@@ -296,10 +296,10 @@ void SDeviceProcesses::UpdateProcessTree()
 	// build process tree
 	if (ShowProcessTree)
 	{
-		for (TMap<uint32, TSharedPtr<FDeviceProcessesProcessTreeNode>>::TConstIterator It(ProcessMap); It; ++It)
+		for (SDeviceProcesses::ProcessMapType::TConstIterator It(ProcessMap); It; ++It)
 		{
 			TSharedPtr<FDeviceProcessesProcessTreeNode> Node = It.Value();
-			TSharedPtr<FDeviceProcessesProcessTreeNode> Parent = ProcessMap.FindRef(Node->GetProcessInfo().ParentId);
+			TSharedPtr<FDeviceProcessesProcessTreeNode> Parent = ProcessMap.FindRef((int64)Node->GetProcessInfo().ParentId);
 
 			if (Parent.IsValid())
 			{
@@ -312,7 +312,7 @@ void SDeviceProcesses::UpdateProcessTree()
 	// filter process list
 	ProcessList.Reset();
 
-	for (TMap<uint32, TSharedPtr<FDeviceProcessesProcessTreeNode>>::TConstIterator It(ProcessMap); It; ++It)
+	for (SDeviceProcesses::ProcessMapType::TConstIterator It(ProcessMap); It; ++It)
 	{
 		TSharedPtr<FDeviceProcessesProcessTreeNode> Node = It.Value();
 
diff --git a/Engine/Source/Developer/DeviceManager/Private/Widgets/Processes/SDeviceProcesses.h b/Engine/Source/Developer/DeviceManager/Private/Widgets/Processes/SDeviceProcesses.h
--- a/Engine/Source/Developer/DeviceManager/Private/Widgets/Processes/SDeviceProcesses.h
+++ b/Engine/Source/Developer/DeviceManager/Private/Widgets/Processes/SDeviceProcesses.h
@@ -45,6 +45,8 @@ public:
 	 */
 	void Construct(const FArguments& InArgs, const TSharedRef<FDeviceManagerModel>& InModel);
 
+	typedef TMap<int64, TSharedPtr<FDeviceProcessesProcessTreeNode>> ProcessMapType;
+
 protected:
 
 	/**
@@ -79,7 +81,7 @@ private:
 	TArray<TSharedPtr<FDeviceProcessesProcessTreeNode>> ProcessList;
 
 	/** Holds the collection of process nodes. */
-	TMap<uint32, TSharedPtr<FDeviceProcessesProcessTreeNode>> ProcessMap;
+	ProcessMapType ProcessMap;
 
 	/** Holds the application tree view. */
 	TSharedPtr<STreeView<TSharedPtr<FDeviceProcessesProcessTreeNode>>> ProcessTreeView;
diff --git a/Engine/Source/Editor/DataTableEditor/Private/DataTableEditor.cpp b/Engine/Source/Editor/DataTableEditor/Private/DataTableEditor.cpp
--- a/Engine/Source/Editor/DataTableEditor/Private/DataTableEditor.cpp
+++ b/Engine/Source/Editor/DataTableEditor/Private/DataTableEditor.cpp
@@ -690,7 +690,7 @@ void FDataTableEditor::RefreshRowNumberColumnWidth()
 
 	for (const FDataTableEditorRowListViewDataPtr& RowData : AvailableRows)
 	{
-		const float RowNumberWidth = FontMeasure->Measure(FString::FromInt(RowData->RowNum), CellTextStyle.Font).X + CellPadding;
+		const float RowNumberWidth = (float)FontMeasure->Measure(FString::FromInt(RowData->RowNum), CellTextStyle.Font).X + CellPadding;
 		RowNumberColumnWidth = FMath::Max(RowNumberColumnWidth, RowNumberWidth);
 	}
 
@@ -715,7 +715,7 @@ void FDataTableEditor::RefreshRowNameColumnWidth()
 
 	for (const FDataTableEditorRowListViewDataPtr& RowData : AvailableRows)
 	{
-		const float RowNameWidth = FontMeasure->Measure(RowData->DisplayName, CellTextStyle.Font).X + CellPadding;
+		const float RowNameWidth = (float)FontMeasure->Measure(RowData->DisplayName, CellTextStyle.Font).X + CellPadding;
 		RowNameColumnWidth = FMath::Max(RowNameColumnWidth, RowNameWidth);
 	}
 	
diff --git a/Engine/Source/Runtime/Slate/Public/Framework/MarqueeRect.h b/Engine/Source/Runtime/Slate/Public/Framework/MarqueeRect.h
--- a/Engine/Source/Runtime/Slate/Public/Framework/MarqueeRect.h
+++ b/Engine/Source/Runtime/Slate/Public/Framework/MarqueeRect.h
@@ -60,6 +60,6 @@ struct FMarqueeRect
 	/** @return This marquee rectangle as a well-formed SlateRect */
 	FSlateRect ToSlateRect() const
 	{
-		return FSlateRect( FMath::Min(StartPoint.X, EndPoint.X), FMath::Min(StartPoint.Y, EndPoint.Y), FMath::Max(StartPoint.X, EndPoint.X), FMath::Max( StartPoint.Y, EndPoint.Y ) );
+		return FSlateRect( FVector2D(FMath::Min(StartPoint.X, EndPoint.X), FMath::Min(StartPoint.Y, EndPoint.Y)), FVector2D(FMath::Max(StartPoint.X, EndPoint.X), FMath::Max( StartPoint.Y, EndPoint.Y )) );
 	}
 };
diff --git a/Engine/Source/Runtime/Slate/Public/Widgets/Layout/SScrollBox.h b/Engine/Source/Runtime/Slate/Public/Widgets/Layout/SScrollBox.h
--- a/Engine/Source/Runtime/Slate/Public/Widgets/Layout/SScrollBox.h
+++ b/Engine/Source/Runtime/Slate/Public/Widgets/Layout/SScrollBox.h
@@ -263,7 +263,7 @@ private:
 	/** Gets the component of a vector in the direction of scrolling based on the Orientation property. */
 	FORCEINLINE float GetScrollComponentFromVector(FVector2D Vector) const
 	{
-		return Orientation == Orient_Vertical ? Vector.Y : Vector.X;
+		return float(Orientation == Orient_Vertical ? Vector.Y : Vector.X);
 	}
 
 	/** Sets the component of a vector in the direction of scrolling based on the Orientation property. */
diff --git a/Engine/Source/Runtime/Slate/Public/Widgets/Views/STableRow.h b/Engine/Source/Runtime/Slate/Public/Widgets/Views/STableRow.h
--- a/Engine/Source/Runtime/Slate/Public/Widgets/Views/STableRow.h
+++ b/Engine/Source/Runtime/Slate/Public/Widgets/Views/STableRow.h
@@ -725,10 +725,10 @@ public:
 	/** @return the zone (above, onto, below) based on where the user is hovering over within the row */
 	EItemDropZone ZoneFromPointerPosition(FVector2D LocalPointerPos, FVector2D LocalSize, EOrientation Orientation)
 	{
-		const float PointerPos = Orientation == EOrientation::Orient_Horizontal ? LocalPointerPos.X : LocalPointerPos.Y;
-		const float Size = Orientation == EOrientation::Orient_Horizontal ? LocalSize.X : LocalSize.Y;
+		const FVector2D::FReal PointerPos = Orientation == EOrientation::Orient_Horizontal ? LocalPointerPos.X : LocalPointerPos.Y;
+		const FVector2D::FReal Size = Orientation == EOrientation::Orient_Horizontal ? LocalSize.X : LocalSize.Y;
 
-		const float ZoneBoundarySu = FMath::Clamp(Size * 0.25f, 3.0f, 10.0f);
+		const FVector2D::FReal ZoneBoundarySu = FMath::Clamp(Size * 0.25f, 3.0f, 10.0f);
 		if (PointerPos < ZoneBoundarySu)
 		{
 			return EItemDropZone::AboveItem;
diff --git a/Engine/Source/Runtime/Slate/Public/Widgets/Views/STreeView.h b/Engine/Source/Runtime/Slate/Public/Widgets/Views/STreeView.h
--- a/Engine/Source/Runtime/Slate/Public/Widgets/Views/STreeView.h
+++ b/Engine/Source/Runtime/Slate/Public/Widgets/Views/STreeView.h
@@ -747,7 +747,7 @@ public:
 			int32 NumPinnedItems = (this->GetNumPinnedItems() < CurrentMaxPinnedItems) ? this->GetNumPinnedItems() : CurrentMaxPinnedItems;
 
 			// This is the first item that would be visible, if there were no pinned rows
-			int32 FirstItem = Results.NewScrollOffset - NumPinnedItems;
+			int32 FirstItem = FMath::TruncToInt32(Results.NewScrollOffset - NumPinnedItems);
 
 			// We find items that have a hierarchy big enough to cover themselves, but select the smallest among them
 			int32 MinSpaceOccupied = TNumericLimits<int32>::Max();
